Copying //UE4/Dev-Sequencer to Dev-Main (//UE4/Dev-Main)

==========================
MAJOR FEATURES + CHANGES
==========================

Change 2800717 on 2015/12/11 by [REDACTED_PERSON]

	Sequencer: Sort the key times for drawing to fix path trajectory.

Change 2803299 on 2015/12/15 by [REDACTED_PERSON]

	Sequencer: Fix property names so that they're the display names. For example, "DepthOfFieldFStop" now reads as "Aperture F Stop"

Change 2804586 on 2015/12/15 by [REDACTED_PERSON]

	Sequencer: Add zoom in/out with shortcuts underscore and equals.

Change 2811823 on 2015/12/23 by [REDACTED_PERSON]

	Editor: Added UI action for creating new content browser folders; code cleanup; removed dead code

	Based on GitHub PR #1809 by artemavrin (https://github.com/EpicGames/UnrealEngine/pull/1809)

	#github: 1809

Change 2811839 on 2015/12/23 by [REDACTED_PERSON]

	StereoPanorama: Code cleanup pass

	Based on GitHub PR# 1756 by ETayrienHBO (https://github.com/EpicGames/UnrealEngine/pull/1756)

	Also:
	- NULL to nullptr
	- namespaced enums to enum classes
	- consistent whitespace, line breaks and parentheses

	#github: 1756

Change 2819172 on 2016/01/07 by [REDACTED_PERSON]

	Sequencer: Marquee and move modes are now automatically activated based on sequencer hotspot

Change 2819176 on 2016/01/07 by [REDACTED_PERSON]

	Sequencer: Various cosmetic fixes

	  - Added icons to tracks
	  - Removed SAnimationOutlinerTreeNode dependency from FSequencerDisplayNode (to enable future customization of shot/event track etc)
	  - Added spacer nodes between top level display nodes
	  - Various hover states and highlights

Change 2819445 on 2016/01/07 by [REDACTED_PERSON]

	Sequencer: Rendering out a capture from the composition graph now renders at the correct size even if r.ScreenPercentage is not 100.

Change 2820747 on 2016/01/08 by [REDACTED_PERSON]

	Sequencer: Added option to close the editor when capturing starts

Change 2827701 on 2016/01/13 by [REDACTED_PERSON]

	Media: Updating audio track specs each frame to better support streaming media and variable streams.

Change 2828465 on 2016/01/14 by [REDACTED_PERSON]

	Media: Better visualization of unknown media durations

Change 2828469 on 2016/01/14 by [REDACTED_PERSON]

	Media: Checking URL scheme on URLs that didn't pass the file extension filter

Change 2834888 on 2016/01/19 by [REDACTED_PERSON]

	Core: TQueue modernization pass

Change 2834934 on 2016/01/19 by [REDACTED_PERSON]

	Core: Implemented TTripleBuffer for triple buffers.

Change 2834950 on 2016/01/19 by [REDACTED_PERSON]

	Core: Added unit tests for TTripleBuffer dirty flag

Change 2835488 on 2016/01/20 by [REDACTED_PERSON]

	Core: More descriptive method names, initialization constructor, unit tests for TTripleBuffer

Change 2837515 on 2016/01/20 by [REDACTED_PERSON]

	Sequencer: Command line options for custom passes.

Change 2837517 on 2016/01/20 by [REDACTED_PERSON]

	Sequencer: Fix crash in visibility track instance on PIE.

Change 2837518 on 2016/01/20 by [REDACTED_PERSON]

	Sequencer: Add option to lock to frame rate while playing.

Change 2837523 on 2016/01/20 by [REDACTED_PERSON]

	Sequencer: Capture thumbnail on level sequence asset save.

Change 2837527 on 2016/01/20 by [REDACTED_PERSON]

	Sequencer: Added preroll for subsequences. Refactor instance update to combine data in EMovieSceneUpdateData.

Change 2837537 on 2016/01/20 by [REDACTED_PERSON]

	Sequencer: Add sequencer transport controls back into viewports.

Change 2837561 on 2016/01/20 by [REDACTED_PERSON]

	Sequencer: Added ability to convert a possessable to a spawnable

	  - This option is available for any root-level possessable object bindings
	  - It will currently delete the existing possessable (we could make this behaviour optional in future)
	  - There is currently no check to sett if the actor is possessed by subsequent sub-sequences. If this is the case, using a possessable, or externally owned spawnable would be a better bet.

Change 2837565 on 2016/01/20 by [REDACTED_PERSON]

diff --git a/Engine/Config/BaseEngine.ini b/Engine/Config/BaseEngine.ini
--- a/Engine/Config/BaseEngine.ini
+++ b/Engine/Config/BaseEngine.ini
@@ -584,7 +584,6 @@ bUseStreamingPause=false
 +ActiveClassRedirects=(OldClassName="MovieSceneColorSection",NewClassName="/Script/MovieSceneTracks.MovieSceneColorSection")
 +ActiveClassRedirects=(OldClassName="MovieSceneFloatSection",NewClassName="/Script/MovieSceneTracks.MovieSceneFloatSection")
 +ActiveClassRedirects=(OldClassName="MovieSceneParticleSection",NewClassName="/Script/MovieSceneTracks.MovieSceneParticleSection")
-+ActiveClassRedirects=(OldClassName="MovieSceneShotSection",NewClassName="/Script/MovieSceneTracks.MovieSceneShotSection")
 +ActiveClassRedirects=(OldClassName="MovieSceneVectorSection",NewClassName="/Script/MovieSceneTracks.MovieSceneVectorSection")
 +ActiveClassRedirects=(OldClassName="SubMovieSceneSection",NewClassName="/Script/MovieSceneTracks.SubMovieSceneSection")
 
@@ -600,6 +599,12 @@ bUseStreamingPause=false
 +ActiveClassRedirects=(OldClassName="MovieSceneVectorTrack",NewClassName="/Script/MovieSceneTracks.MovieSceneVectorTrack")
 +ActiveClassRedirects=(OldClassName="SubMovieSceneTrack",NewClassName="/Script/MovieSceneTracks.SubMovieSceneTrack")
 
+; MovieSceneShot* renamed to MovieSceneCameraCut*
++ActiveClassRedirects=(OldClassName="MovieSceneShotTrack",NewClassName="/Script/MovieSceneTracks.MovieSceneCameraCutTrack")
++ActiveClassRedirects=(OldClassName="MovieSceneShotSection",NewClassName="/Script/MovieSceneTracks.MovieSceneCameraCutSection")
++TaggedPropertyRedirects=(ClassName="MovieSceneShotSection",OldPropertyName="ShotNumber",NewPropertyName="CameraCutNumber")
++TaggedPropertyRedirects=(ClassName="MovieScene",OldPropertyName="ShotTrack",NewPropertyName="CameraCutTrack")
+
 ; We renamed EditorGameAgnosticSettings to be less imposing
 +ActiveClassRedirects=(OldClassName="EditorGameAgnosticSettings",NewClassName="/Script/UnrealEd.EditorSettings")
 +ActiveClassRedirects=(OldClassName="EditorUserSettings",NewClassName="/Script/UnrealEd.EditorPerProjectUserSettings")
diff --git a/Engine/Config/BaseInput.ini b/Engine/Config/BaseInput.ini
--- a/Engine/Config/BaseInput.ini
+++ b/Engine/Config/BaseInput.ini
@@ -172,3 +172,6 @@ MaxScrollbackSize=1024
 +ManualAutoCompleteList=(Command="Log ",Desc="<category> <level> Change verbosity level for a log category")
 +ManualAutoCompleteList=(Command="Log list",Desc="<search string> Search for log categories")
 +ManualAutoCompleteList=(Command="Log reset",Desc="Undo any changes to log verbosity")
++ManualAutoCompleteList=(Command="RecordAnimation ActorName AnimName",Desc="Record animation for a specified actor to the specified file")
++ManualAutoCompleteList=(Command="StopRecordingAnimation All",Desc="Stop all recording animations")
++ManualAutoCompleteList=(Command="StopRecordingAnimation ActorName",Desc="Stop recording the specifed animation")
diff --git a/Engine/Plugins/Editor/SequenceRecorder/SequenceRecorder.uplugin b/Engine/Plugins/Editor/SequenceRecorder/SequenceRecorder.uplugin
new file mode 100644
--- /dev/null
+++ b/Engine/Plugins/Editor/SequenceRecorder/SequenceRecorder.uplugin
@@ -0,0 +1,25 @@
+{
+	"FileVersion" : 3,
+	"Version" : 1,
+	"VersionName" : "1.0",
+	"FriendlyName" : "Sequence Recorder Plugin",
+	"Description" : "Records sequences from actors",
+	"Category" : "Editor",
+	"CreatedBy" : "Epic Games, Inc.",
+	"CreatedByURL" : "http://epicgames.com",
+	"DocsURL" : "",
+	"MarketplaceURL" : "",
+	"SupportURL" : "",
+	"EnabledByDefault" : false,
+	"CanContainContent" : false,
+	"IsBetaVersion" : true,
+	"Installed" : false,
+	"Modules" :
+	[
+		{
+			"Name" : "SequenceRecorder",
+			"Type" : "Editor",
+			"LoadingPhase" : "PostEngineInit"
+		}
+	]
+}
\ No newline at end of file
diff --git a/Engine/Plugins/Editor/SequenceRecorder/Source/SequenceRecorder/Private/ActorRecording.cpp b/Engine/Plugins/Editor/SequenceRecorder/Source/SequenceRecorder/Private/ActorRecording.cpp
new file mode 100644
--- /dev/null
+++ b/Engine/Plugins/Editor/SequenceRecorder/Source/SequenceRecorder/Private/ActorRecording.cpp
@@ -0,0 +1,116 @@
+// Copyright 1998-2016 Epic Games, Inc. All Rights Reserved.
+
+#include "SequenceRecorderPrivatePCH.h"
+#include "ActorRecording.h"
+#include "SequenceRecorderSettings.h"
+#include "SequenceRecorderUtils.h"
+
+bool UActorRecording::StartRecording()
+{
+	LastRecordedAnimation = nullptr;
+
+	if(ActorToRecord.IsValid())
+	{
+		// @todo: animation recording: allow recording of multiple skeletal meshes per actor
+		USkeletalMeshComponent* SkeletalMeshComponent = ActorToRecord.Get()->FindComponentByClass<USkeletalMeshComponent>();
+		if(SkeletalMeshComponent)
+		{
+			// turn off URO and make sure we always update even if out of view
+			bEnableUpdateRateOptimizations = SkeletalMeshComponent->bEnableUpdateRateOptimizations;
+			MeshComponentUpdateFlag = SkeletalMeshComponent->MeshComponentUpdateFlag;
+
+			SkeletalMeshComponent->bEnableUpdateRateOptimizations = false;
+			SkeletalMeshComponent->MeshComponentUpdateFlag = EMeshComponentUpdateFlag::AlwaysTickPoseAndRefreshBones;
+
+			UAnimSequence* Sequence = nullptr;
+			const bool bValidSpecifiedAnimation = (bSpecifyTargetAnimation && TargetAnimation.IsValid());
+			if(bValidSpecifiedAnimation)
+			{
+				Sequence = TargetAnimation.Get();
+			}
+
+			bool bValidAutoGeneratedAnimation = false;
+			if(!bSpecifyTargetAnimation)
+			{
+				// build an asset path
+				const USequenceRecorderSettings* Settings = GetDefault<USequenceRecorderSettings>();
+
+				FString AssetPath(TEXT("/Game"));
+				AssetPath /= Settings->SequenceRecordingBasePath.Path;
+				if(Settings->AnimationSubDirectory.Len() > 0)
+				{
+					AssetPath /= Settings->AnimationSubDirectory;
+				}
+
+				FString AssetName = Settings->SequenceName.Len() > 0 ? Settings->SequenceName : TEXT("RecordedSequence");
+				AssetName += TEXT("_");
+				AssetName += ActorToRecord->GetActorLabel();
+
+				Sequence = SequenceRecorderUtils::MakeNewAsset<UAnimSequence>(AssetPath, AssetName);
+				if(Sequence)
+				{
+					// set skeleton
+					Sequence->SetSkeleton(SkeletalMeshComponent->SkeletalMesh->Skeleton);
+					bValidAutoGeneratedAnimation = true;
+				}
+			}
+
+			if(bValidSpecifiedAnimation || bValidAutoGeneratedAnimation)
+			{
+				FAnimationRecorderManager::Get().RecordAnimation(SkeletalMeshComponent, Sequence, AnimationSettings);
+				LastRecordedAnimation = Sequence;
+				return true;
+			}
+		}
+	}
+
+	return false;
+}
+
+void UActorRecording::Tick(float DeltaSeconds)
+{
+	if (ActorToRecord.IsValid())
+	{
+		USkeletalMeshComponent* SkeletalMeshComponent = ActorToRecord.Get()->FindComponentByClass<USkeletalMeshComponent>();
+		if (SkeletalMeshComponent)
+		{
+			// re-force updates on as gameplay can sometimes turn these back off!
+			SkeletalMeshComponent->bEnableUpdateRateOptimizations = false;
+			SkeletalMeshComponent->MeshComponentUpdateFlag = EMeshComponentUpdateFlag::AlwaysTickPoseAndRefreshBones;
+		}
+	}
+}
+
+bool UActorRecording::StopRecording()
+{
+	if (ActorToRecord.IsValid())
+	{
+		// @todo: animation recording: allow recording of multiple skeletal meshes per actor
+		USkeletalMeshComponent* SkeletalMeshComponent = ActorToRecord.Get()->FindComponentByClass<USkeletalMeshComponent>();
+		if (SkeletalMeshComponent)
+		{
+			// restore update flags
+			SkeletalMeshComponent->bEnableUpdateRateOptimizations = bEnableUpdateRateOptimizations;
+			SkeletalMeshComponent->MeshComponentUpdateFlag = MeshComponentUpdateFlag;
+
+			FAnimationRecorderManager::Get().StopRecordingAnimation(SkeletalMeshComponent);
+			return true;
+		}
+	}
+
+	return false;
+}
+
+bool UActorRecording::IsRecording() const
+{
+	if(ActorToRecord.IsValid())
+	{
+		USkeletalMeshComponent* SkeletalMeshComponent = ActorToRecord.Get()->FindComponentByClass<USkeletalMeshComponent>();
+		if (SkeletalMeshComponent)
+		{
+			return FAnimationRecorderManager::Get().IsRecording(SkeletalMeshComponent);
+		}
+	}
+
+	return false;
+}
diff --git a/Engine/Plugins/Editor/SequenceRecorder/Source/SequenceRecorder/Private/ActorRecording.h b/Engine/Plugins/Editor/SequenceRecorder/Source/SequenceRecorder/Private/ActorRecording.h
new file mode 100644
--- /dev/null
+++ b/Engine/Plugins/Editor/SequenceRecorder/Source/SequenceRecorder/Private/ActorRecording.h
@@ -0,0 +1,94 @@
+// Copyright 1998-2016 Epic Games, Inc. All Rights Reserved.
+
+#pragma once
+
+#include "AnimationRecordingSettings.h"
+#include "Components/SkinnedMeshComponent.h"
+
+#include "ActorRecording.generated.h"
+
+USTRUCT()
+struct FActorRecordingSettings
+{
+	GENERATED_BODY()
+
+	FActorRecordingSettings()
+		: bRecordTransforms(false)
+		, bRecordVisibility(true)
+		, bRecordNearbySpawnedActors(false)
+		, NearbyActorRecordingProximity(1000.0f)
+	{}
+
+public:
+	/** 
+	 * Whether to record actor transforms. This can be useful if you want the actors to end up in specific locations after the sequence. 
+	 * By default we rely on animations to provide transforms, but this can be changed using the "Record In World Space" animation setting.
+	 */
+	UPROPERTY(EditAnywhere, Category = "Actor Recording")
+	bool bRecordTransforms;
+
+	/** Whether to record actor visibility. */
+	UPROPERTY(EditAnywhere, Category = "Actor Recording")
+	bool bRecordVisibility;
+
+	/** Whether to record nearby spawned actors. */
+	UPROPERTY(EditAnywhere, Category = "Actor Recording")
+	bool bRecordNearbySpawnedActors;
+
+	/** Proximity to currently recorded actors to record newly spawned actors. */
+	UPROPERTY(EditAnywhere, Category = "Actor Recording")
+	float NearbyActorRecordingProximity;
+};
+
+UCLASS(MinimalAPI, Transient)
+class UActorRecording : public UObject
+{
+	GENERATED_BODY()
+
+public:
+	/** Start this queued recording */
+	bool StartRecording();
+
+	/** Stop this recording. Has no effect if we are not currently recording */
+	bool StopRecording();
+
+	/** Tick this recording */
+	void Tick(float DeltaSeconds);
+
+	/** Whether we are currently recording */
+	bool IsRecording() const;
+
+public:
+	/** The actor we want to record */
+	UPROPERTY(EditAnywhere, Category = "Actor Recording")
+	TLazyObjectPtr<AActor> ActorToRecord;
+
+	UPROPERTY(EditAnywhere, Category = "Actor Recording")
+	FActorRecordingSettings ActorSettings;
+
+	/** Whether we should specify the target animation or auto-create it */
+	UPROPERTY(EditAnywhere, Category = "Animation Recording")
+	bool bSpecifyTargetAnimation;
+
+	/** The target animation we want to record to */
+	UPROPERTY(EditAnywhere, Category = "Animation Recording", meta=(EditCondition = "bSpecifyTargetAnimation"))
+	TWeakObjectPtr<UAnimSequence> TargetAnimation;
+
+	/** The settings to apply to this actor's animation */
+	UPROPERTY(EditAnywhere, Category = "Animation Recording")
+	FAnimationRecordingSettings AnimationSettings;
+
+	/** Our last recorded animation */
+	UPROPERTY(Transient)
+	TWeakObjectPtr<UAnimSequence> LastRecordedAnimation;
+
+	/** Guid that identifies our spawnable in a recorded sequence */
+	FGuid Guid;
+
+private:
+	/** Used to store/restore update flag when recording */
+	EMeshComponentUpdateFlag::Type MeshComponentUpdateFlag;
+
+	/** Used to store/restore URO when recording */
+	bool bEnableUpdateRateOptimizations;
+};
\ No newline at end of file
diff --git a/Engine/Plugins/Editor/SequenceRecorder/Source/SequenceRecorder/Private/AnimationRecorder.cpp b/Engine/Plugins/Editor/SequenceRecorder/Source/SequenceRecorder/Private/AnimationRecorder.cpp
new file mode 100644
--- /dev/null
+++ b/Engine/Plugins/Editor/SequenceRecorder/Source/SequenceRecorder/Private/AnimationRecorder.cpp
@@ -0,0 +1,821 @@
+// Copyright 1998-2016 Epic Games, Inc. All Rights Reserved.
+
+#include "SequenceRecorderPrivatePCH.h"
+#include "AnimationRecorder.h"
+#include "Animation/AnimCompress.h"
+#include "Animation/AnimCompress_BitwiseCompressOnly.h"
+#include "SCreateAnimationDlg.h"
+#include "Runtime/AssetRegistry/Public/AssetRegistryModule.h"
+#include "SNotificationList.h"
+#include "NotificationManager.h"
+#include "EngineLogs.h"
+#include "Animation/AnimationSettings.h"
+#include "Toolkits/AssetEditorManager.h"
+#include "ActorRecording.h"
+
+#define LOCTEXT_NAMESPACE "FAnimationRecorder"
+
+/////////////////////////////////////////////////////
+
+static TAutoConsoleVariable<float> CVarAnimRecorderSampleRate(
+	TEXT("AnimRecorder.SampleRate"),
+	DEFAULT_SAMPLERATE,
+	TEXT("Sets the sample rate for the animation recorder system"),
+	ECVF_Default);
+
+static TAutoConsoleVariable<int32> CVarAnimRecorderWorldSpace(
+	TEXT("AnimRecorder.RecordInWorldSpace"),
+	1,
+	TEXT("True to record anim keys in world space, false to record only in local space."),
+	ECVF_Default);
+
+FAnimationRecorder::FAnimationRecorder()
+	: AnimationObject(nullptr)
+	, bRecordLocalToWorld(false)
+	, bAutoSaveAsset(false)
+
+{
+	SetSampleRateAndLength(FAnimationRecordingSettings::DefaultSampleRate, FAnimationRecordingSettings::DefaultMaximumLength);
+}
+
+FAnimationRecorder::~FAnimationRecorder()
+{
+	StopRecord(false);
+}
+
+void FAnimationRecorder::SetSampleRateAndLength(float SampleRateHz, float LengthInSeconds)
+{
+	if (SampleRateHz <= 0.f)
+	{
+		// invalid rate passed in, fall back to default
+		SampleRateHz = FAnimationRecordingSettings::DefaultSampleRate;
+	}
+
+	if (LengthInSeconds <= 0.f)
+	{
+		// invalid length passed in, default to unbounded
+		SampleRateHz = FAnimationRecordingSettings::DefaultSampleRate;
+	}
+
+	IntervalTime = 1.0f / SampleRateHz;
+	if (LengthInSeconds == FAnimationRecordingSettings::UnboundedMaximumLength)
+	{
+		// invalid length passed in, default to unbounded
+		MaxFrame = UnBoundedFrameCount;
+	}
+	else
+	{
+		MaxFrame = SampleRateHz * LengthInSeconds;
+	}
+}
+
+bool FAnimationRecorder::SetAnimCompressionScheme(TSubclassOf<UAnimCompress> SchemeClass)
+{
+	if (AnimationObject)
+	{
+		UAnimCompress* const SchemeObject = NewObject<UAnimCompress>(GetTransientPackage(), SchemeClass);
+		if (SchemeObject)
+		{
+			AnimationObject->CompressionScheme = SchemeObject;
+			return true;
+		}
+	}
+
+	return false;
+}
+
+// Internal. Pops up a dialog to get saved asset path
+static bool PromptUserForAssetPath(FString& AssetPath, FString& AssetName)
+{
+	TSharedRef<SCreateAnimationDlg> NewAnimDlg = SNew(SCreateAnimationDlg);
+	if (NewAnimDlg->ShowModal() != EAppReturnType::Cancel)
+	{
+		AssetPath = NewAnimDlg->GetFullAssetPath();
+		AssetName = NewAnimDlg->GetAssetName();
+		return true;
+	}
+
+	return false;
+}
+
+bool FAnimationRecorder::TriggerRecordAnimation(USkeletalMeshComponent* Component)
+{
+	FString AssetPath;
+	FString AssetName;
+
+	if (!Component || !Component->SkeletalMesh || !Component->SkeletalMesh->Skeleton)
+	{
+		return false;
+	}
+
+	// ask for path
+	if (PromptUserForAssetPath(AssetPath, AssetName))
+	{
+		return TriggerRecordAnimation(Component, AssetPath, AssetName);
+	}
+
+	return false;
+}
+
+bool FAnimationRecorder::TriggerRecordAnimation(USkeletalMeshComponent* Component, const FString& InAssetPath, const FString& InAssetName)
+{
+	if (!Component || !Component->SkeletalMesh || !Component->SkeletalMesh->Skeleton)
+	{
+		return false;
+	}
+
+	// create the asset
+	FText InvalidPathReason;
+	bool const bValidPackageName = FPackageName::IsValidLongPackageName(InAssetPath, false, &InvalidPathReason);
+	if (bValidPackageName == false)
+	{
+		UE_LOG(LogAnimation, Log, TEXT("%s is an invalid asset path, prompting user for new asset path. Reason: %s"), *InAssetPath, *InvalidPathReason.ToString());
+	}
+
+	FString ValidatedAssetPath = InAssetPath;
+	FString ValidatedAssetName = InAssetName;
+
+	UObject* Parent = bValidPackageName ? CreatePackage(nullptr, *ValidatedAssetPath) : nullptr;
+	if (Parent == nullptr)
+	{
+		// bad or no path passed in, do the popup
+		if (PromptUserForAssetPath(ValidatedAssetPath, ValidatedAssetName) == false)
+		{
+			return false;
+		}
+
+		Parent = CreatePackage(nullptr, *ValidatedAssetPath);
+	}
+
+	UObject* const Object = LoadObject<UObject>(Parent, *ValidatedAssetName, nullptr, LOAD_None, nullptr);
+	// if object with same name exists, warn user
+	if (Object)
+	{
+		FMessageDialog::Open(EAppMsgType::Ok, NSLOCTEXT("UnrealEd", "Error_AssetExist", "Asset with same name exists. Can't overwrite another asset"));
+		return false;		// failed
+	}
+
+	// If not, create new one now.
+	UAnimSequence* const NewSeq = NewObject<UAnimSequence>(Parent, *ValidatedAssetName, RF_Public | RF_Standalone);
+	if (NewSeq)
+	{
+		// set skeleton
+		NewSeq->SetSkeleton(Component->SkeletalMesh->Skeleton);
+		// Notify the asset registry
+		FAssetRegistryModule::AssetCreated(NewSeq);
+		StartRecord(Component, NewSeq);
+
+		return true;
+	}
+
+	return false;
+}
+
+
+void FAnimationRecorder::StartRecord(USkeletalMeshComponent* Component, UAnimSequence* InAnimationObject)
+{
+	TimePassed = 0.f;
+	AnimationObject = InAnimationObject;
+
+	AnimationObject->RawAnimationData.Empty();
+	AnimationObject->TrackToSkeletonMapTable.Empty();
+	AnimationObject->AnimationTrackNames.Empty();
+
+	PreviousSpacesBases = Component->GetSpaceBases();
+	PreviousComponentToWorld = Component->ComponentToWorld;
+
+	LastFrame = 0;
+	AnimationObject->SequenceLength = 0.f;
+	AnimationObject->NumFrames = 0;
+
+	USkeleton* AnimSkeleton = AnimationObject->GetSkeleton();
+	// add all frames
+	for (int32 BoneIndex=0; BoneIndex <PreviousSpacesBases.Num(); ++BoneIndex)
+	{
+		// verify if this bone exists in skeleton
+		int32 BoneTreeIndex = AnimSkeleton->GetSkeletonBoneIndexFromMeshBoneIndex(Component->SkeletalMesh, BoneIndex);
+		if (BoneTreeIndex != INDEX_NONE)
+		{
+			// add tracks for the bone existing
+			FName BoneTreeName = AnimSkeleton->GetReferenceSkeleton().GetBoneName(BoneTreeIndex);
+			AnimationObject->AnimationTrackNames.Add(BoneTreeName);
+			// add mapping to skeleton bone track
+			AnimationObject->TrackToSkeletonMapTable.Add(FTrackToSkeletonMap(BoneTreeIndex));
+		}
+	}
+
+	int32 NumTracks = AnimationObject->AnimationTrackNames.Num();
+	// add tracks
+	AnimationObject->RawAnimationData.AddZeroed(NumTracks);
+
+	// init notifies
+	AnimationObject->InitializeNotifyTrack();
+
+	// record the first frame
+	Record(Component, PreviousComponentToWorld, PreviousSpacesBases, 0);
+}
+
+void FAnimationRecorder::FixupNotifies()
+{
+	if (AnimationObject)
+	{
+		// build notify tracks - first find how many tracks we want
+		for (FAnimNotifyEvent& Event : AnimationObject->Notifies)
+		{
+			if (Event.TrackIndex >= AnimationObject->AnimNotifyTracks.Num())
+			{
+				AnimationObject->AnimNotifyTracks.SetNum(Event.TrackIndex + 1);
+
+				// remake track names to create a nice sequence
+				const int32 TrackNum = AnimationObject->AnimNotifyTracks.Num();
+				for (int32 TrackIndex = 0; TrackIndex < TrackNum; ++TrackIndex)
+				{
+					FAnimNotifyTrack& Track = AnimationObject->AnimNotifyTracks[TrackIndex];
+					Track.TrackName = *FString::FromInt(TrackIndex + 1);
+				}
+			}
+		}
+
+		// now build tracks
+		for (int32 EventIndex = 0; EventIndex < AnimationObject->Notifies.Num(); ++EventIndex)
+		{
+			FAnimNotifyEvent& Event = AnimationObject->Notifies[EventIndex];
+			AnimationObject->AnimNotifyTracks[Event.TrackIndex].Notifies.Add(&AnimationObject->Notifies[EventIndex]);
+		}
+	}
+}
+
+UAnimSequence* FAnimationRecorder::StopRecord(bool bShowMessage)
+{
+	if (AnimationObject)
+	{
+		int32 NumFrames = LastFrame  + 1;
+		AnimationObject->NumFrames = NumFrames;
+
+		// can't use TimePassed. That is just total time that has been passed, not necessarily match with frame count
+		AnimationObject->SequenceLength = (NumFrames>1) ? (NumFrames-1) * IntervalTime : MINIMUM_ANIMATION_LENGTH;
+
+		FixupNotifies();
+
+		// force anim settings for speed, we dont want any fancy recompression at present
+		UAnimationSettings* AnimationSettings = GetMutableDefault<UAnimationSettings>();
+		TSubclassOf<UAnimCompress> OldDefaultCompressionAlgorithm = AnimationSettings->DefaultCompressionAlgorithm;
+		TEnumAsByte<AnimationCompressionFormat> OldRotationCompressionFormat = AnimationSettings->RotationCompressionFormat;
+		TEnumAsByte<AnimationCompressionFormat> OldTranslationCompressionFormat = AnimationSettings->TranslationCompressionFormat;
+
+		AnimationSettings->DefaultCompressionAlgorithm = UAnimCompress_BitwiseCompressOnly::StaticClass();
+		AnimationSettings->RotationCompressionFormat = ACF_None;
+		AnimationSettings->TranslationCompressionFormat = ACF_None;
+
+		// post-process applies compression etc.
+		AnimationObject->PostProcessSequence();
+
+		// restore old settings
+		AnimationSettings->DefaultCompressionAlgorithm = OldDefaultCompressionAlgorithm;
+		AnimationSettings->RotationCompressionFormat = OldRotationCompressionFormat;
+		AnimationSettings->TranslationCompressionFormat = OldTranslationCompressionFormat;
+
+		AnimationObject->MarkPackageDirty();
+		
+		// save the package to disk, for convenience and so we can run this in standalone mode
+		if (bAutoSaveAsset)
+		{
+			UPackage* const Package = AnimationObject->GetOutermost();
+			FString const PackageName = Package->GetName();
+			FString const PackageFileName = FPackageName::LongPackageNameToFilename(PackageName, FPackageName::GetAssetPackageExtension());
+
+			UPackage::SavePackage(Package, NULL, RF_Standalone, *PackageFileName, GError, nullptr, false, true, SAVE_NoError);
+		}
+
+		UAnimSequence* ReturnObject = AnimationObject;
+
+		// notify to user
+		if (bShowMessage)
+		{
+			const FText NotificationText = FText::Format(LOCTEXT("RecordAnimation", "'{0}' has been successfully recorded [{1} frames : {2} sec(s) @ {3} Hz]"),
+				FText::FromString(AnimationObject->GetName()),
+				FText::AsNumber(AnimationObject->NumFrames),
+				FText::AsNumber(AnimationObject->SequenceLength),
+				FText::AsNumber(1.f / IntervalTime)
+				);
+					
+			if (GIsEditor)
+			{
+				FNotificationInfo Info(NotificationText);
+				Info.ExpireDuration = 8.0f;
+				Info.bUseLargeFont = false;
+				Info.Hyperlink = FSimpleDelegate::CreateLambda([=]()
+				{
+					TArray<UObject*> Assets;
+					Assets.Add(ReturnObject);
+					FAssetEditorManager::Get().OpenEditorForAssets(Assets);
+				});
+				Info.HyperlinkText = FText::Format(LOCTEXT("OpenNewAnimationHyperlink", "Open {0}"), FText::FromString(AnimationObject->GetName()));
+				TSharedPtr<SNotificationItem> Notification = FSlateNotificationManager::Get().AddNotification(Info);
+				if ( Notification.IsValid() )
+				{
+					Notification->SetCompletionState( SNotificationItem::CS_Success );
+				}
+			}
+
+			FAssetRegistryModule::AssetCreated(AnimationObject);
+		}
+
+		AnimationObject = NULL;
+		PreviousSpacesBases.Empty();
+
+		return ReturnObject;
+	}
+
+	UniqueNotifies.Empty();
+	UniqueNotifyStates.Empty();
+
+	return NULL;
+}
+
+void FAnimationRecorder::UpdateRecord(USkeletalMeshComponent* Component, float DeltaTime)
+{
+	// if no animation object, return
+	if (!AnimationObject || !Component)
+	{
+		return;
+	}
+
+	// no sim time, no record
+	if (DeltaTime <= 0.f)
+	{
+		return;
+	}
+
+	// in-editor we can get a long frame update because of the modal dialog used to pick paths
+	if(DeltaTime > IntervalTime && (LastFrame == 0 || LastFrame == 1))
+	{
+		DeltaTime = IntervalTime;
+	}
+
+	float const PreviousTimePassed = TimePassed;
+	TimePassed += DeltaTime;
+
+	// time passed has been updated
+	// now find what frames we need to update
+	int32 FramesRecorded = LastFrame;
+	int32 FramesToRecord = FPlatformMath::TruncToInt(TimePassed / IntervalTime);
+
+	// notifies need to be done regardless of sample rate
+	if (Component->GetAnimInstance())
+	{
+		RecordNotifies(Component, Component->GetAnimInstance()->NotifyQueue.AnimNotifies, DeltaTime, TimePassed);
+	}
+
+	if (FramesRecorded < FramesToRecord)
+	{
+		const TArray<FTransform>& SpaceBases = Component->GetSpaceBases();
+
+		check(SpaceBases.Num() == PreviousSpacesBases.Num());
+
+		TArray<FTransform> BlendedSpaceBases;
+		BlendedSpaceBases.AddZeroed(SpaceBases.Num());
+
+		UE_LOG(LogAnimation, Log, TEXT("DeltaTime : %0.2f, Current Frame Count : %d, Frames To Record : %d, TimePassed : %0.2f"), DeltaTime
+			, FramesRecorded, FramesToRecord, TimePassed);
+
+		// if we need to record frame
+		while (FramesToRecord > FramesRecorded)
+		{
+			// find what frames we need to record
+			// convert to time
+			float CurrentTime = (FramesRecorded + 1) * IntervalTime;
+			float BlendAlpha = (CurrentTime - PreviousTimePassed) / DeltaTime;
+
+			UE_LOG(LogAnimation, Log, TEXT("Current Frame Count : %d, BlendAlpha : %0.2f"), FramesRecorded + 1, BlendAlpha);
+
+			// for now we just concern component space, not skeleton space
+			for (int32 BoneIndex = 0; BoneIndex<SpaceBases.Num(); ++BoneIndex)
+			{
+				BlendedSpaceBases[BoneIndex].Blend(PreviousSpacesBases[BoneIndex], SpaceBases[BoneIndex], BlendAlpha);
+			}
+
+			FTransform BlendedComponentToWorld;
+			BlendedComponentToWorld.Blend(PreviousComponentToWorld, Component->ComponentToWorld, BlendAlpha);
+
+			Record(Component, BlendedComponentToWorld, BlendedSpaceBases, FramesRecorded + 1);
+			++FramesRecorded;
+		}
+	}
+
+	//save to current transform
+	PreviousSpacesBases = Component->GetSpaceBases();
+	PreviousComponentToWorld = Component->ComponentToWorld;
+
+	// if we passed MaxFrame, just stop it
+	if (MaxFrame != UnBoundedFrameCount && FramesRecorded >= MaxFrame)
+	{
+		UE_LOG(LogAnimation, Log, TEXT("Animation Recording exceeds the time limited (%d mins). Stopping recording animation... "), (int32)((float)MaxFrame / ((1.0f / IntervalTime) * 60.0f)));
+		StopRecord(true);
+	}
+}
+
+void FAnimationRecorder::Record( USkeletalMeshComponent* Component, FTransform const& ComponentToWorld, TArray<FTransform> SpacesBases, int32 FrameToAdd )
+{
+	if (ensure(AnimationObject))
+	{
+		USkeleton* AnimSkeleton = AnimationObject->GetSkeleton();
+		for (int32 TrackIndex=0; TrackIndex <AnimationObject->RawAnimationData.Num(); ++TrackIndex)
+		{
+			// verify if this bone exists in skeleton
+			int32 BoneTreeIndex = AnimationObject->GetSkeletonIndexFromTrackIndex(TrackIndex);
+			if (BoneTreeIndex != INDEX_NONE)
+			{
+				int32 BoneIndex = AnimSkeleton->GetMeshBoneIndexFromSkeletonBoneIndex(Component->SkeletalMesh, BoneTreeIndex);
+				int32 ParentIndex = Component->SkeletalMesh->RefSkeleton.GetParentIndex(BoneIndex);
+				FTransform LocalTransform = SpacesBases[BoneIndex];
+				if ( ParentIndex != INDEX_NONE )
+				{
+					LocalTransform.SetToRelativeTransform(SpacesBases[ParentIndex]);
+				}
+				// if record local to world, we'd like to consider component to world to be in root
+				else if (bRecordLocalToWorld)
+				{
+					LocalTransform *= ComponentToWorld;
+				}
+
+				FRawAnimSequenceTrack& RawTrack = AnimationObject->RawAnimationData[TrackIndex];
+				RawTrack.PosKeys.Add(LocalTransform.GetTranslation());
+				RawTrack.RotKeys.Add(LocalTransform.GetRotation());
+				RawTrack.ScaleKeys.Add(LocalTransform.GetScale3D());
+
+				// verification
+				check (FrameToAdd == RawTrack.PosKeys.Num()-1);
+			}
+		}
+
+		LastFrame = FrameToAdd;
+	}
+}
+
+void FAnimationRecorder::RecordNotifies(USkeletalMeshComponent* Component, const TArray<const struct FAnimNotifyEvent*>& AnimNotifies, float DeltaTime, float RecordTime)
+{
+	if (ensure(AnimationObject))
+	{
+		// flag notifies as possibly unused this frame
+		for (auto& ActiveNotify : ActiveNotifies)
+		{
+			ActiveNotify.Value = false;
+		}
+
+		int32 AddedThisFrame = 0;
+		for(const FAnimNotifyEvent* NotifyEvent : AnimNotifies)
+		{
+			// we don't want to insert notifies with duration more than once
+			if(NotifyEvent->GetDuration() > 0.0f)
+			{
+				// if this event is active already then don't add it
+				bool bAlreadyActive = false;
+				for(auto& ActiveNotify : ActiveNotifies)
+				{
+					if(NotifyEvent == ActiveNotify.Key)
+					{
+						// flag as active
+						ActiveNotify.Value = true;
+						bAlreadyActive = true;
+						break;
+					}
+				}
+
+				// already active, so skip adding
+				if(bAlreadyActive)
+				{
+					continue;
+				}
+				else
+				{
+					// add a new active notify with duration
+					ActiveNotifies.Add(TPairInitializer<const FAnimNotifyEvent*, bool>(NotifyEvent, true));
+				}
+			}
+
+			// make a new notify from this event & set the current time
+			FAnimNotifyEvent NewEvent = *NotifyEvent;
+			NewEvent.SetTime(RecordTime);
+
+			// see if we need to create a new notify
+			if(NotifyEvent->Notify)
+			{
+				UAnimNotify** FoundNotify = UniqueNotifies.Find(NotifyEvent->Notify);
+				if(FoundNotify == nullptr)
+				{
+					NewEvent.Notify = Cast<UAnimNotify>(StaticDuplicateObject(NewEvent.Notify, AnimationObject));
+					UniqueNotifies.Add(NotifyEvent->Notify, NewEvent.Notify);
+				}
+				else
+				{
+					NewEvent.Notify = *FoundNotify;
+				}
+			}
+
+			// see if we need to create a new notify state
+			if (NotifyEvent->NotifyStateClass)
+			{
+				UAnimNotifyState** FoundNotifyState = UniqueNotifyStates.Find(NotifyEvent->NotifyStateClass);
+				if (FoundNotifyState == nullptr)
+				{
+					NewEvent.NotifyStateClass = Cast<UAnimNotifyState>(StaticDuplicateObject(NewEvent.NotifyStateClass, AnimationObject));
+					UniqueNotifyStates.Add(NotifyEvent->NotifyStateClass, NewEvent.NotifyStateClass);
+				}
+				else
+				{
+					NewEvent.NotifyStateClass = *FoundNotifyState;
+				}
+			}
+
+			AnimationObject->Notifies.Add(NewEvent);
+			AddedThisFrame++;
+		}
+
+		// remove all notifies that didnt get added this time
+		ActiveNotifies.RemoveAll([](TPair<const FAnimNotifyEvent*, bool>& ActiveNotify){ return !ActiveNotify.Value; });
+
+		UE_LOG(LogAnimation, Log, TEXT("Added notifies : %d"), AddedThisFrame);
+	}
+}
+
+void FAnimationRecorderManager::Tick(float DeltaTime)
+{
+	for (auto& Inst : RecorderInstances)
+	{
+		Inst.Update(DeltaTime);
+	}
+}
+
+void FAnimationRecorderManager::Tick(USkeletalMeshComponent* Component, float DeltaTime)
+{
+	for (auto& Inst : RecorderInstances)
+	{
+		if(Inst.SkelComp == Component)
+		{
+			Inst.Update(DeltaTime);
+		}
+	}
+}
+
+FAnimationRecorderManager::FAnimationRecorderManager()
+{
+}
+
+FAnimationRecorderManager::~FAnimationRecorderManager()
+{
+}
+
+FAnimationRecorderManager& FAnimationRecorderManager::Get()
+{
+	static FAnimationRecorderManager AnimRecorderManager;
+	return AnimRecorderManager;
+}
+
+
+FAnimRecorderInstance::FAnimRecorderInstance()
+	: SkelComp(nullptr)
+	, Recorder(nullptr)
+{
+}
+
+void FAnimRecorderInstance::Init(USkeletalMeshComponent* InComponent, const FString& InAssetPath, const FString& InAssetName, float SampleRateHz, float MaxLength, bool bRecordInWorldSpace, bool bAutoSaveAsset)
+{
+	SkelComp = InComponent;
+	AssetPath = InAssetPath;
+	AssetName = InAssetName;
+	Recorder = MakeShareable(new FAnimationRecorder());
+	Recorder->SetSampleRateAndLength(SampleRateHz, MaxLength);
+	Recorder->bRecordLocalToWorld = bRecordInWorldSpace;
+	Recorder->SetAnimCompressionScheme(UAnimCompress_BitwiseCompressOnly::StaticClass());
+	Recorder->bAutoSaveAsset = bAutoSaveAsset;
+
+	if (InComponent)
+	{
+		CachedSkelCompForcedLodModel = InComponent->ForcedLodModel;
+		InComponent->ForcedLodModel = 1;
+	}
+}
+
+void FAnimRecorderInstance::Init(USkeletalMeshComponent* InComponent, UAnimSequence* InSequence, float SampleRateHz, float MaxLength, bool bRecordInWorldSpace, bool bAutoSaveAsset)
+{
+	SkelComp = InComponent;
+	Sequence = InSequence;
+	Recorder = MakeShareable(new FAnimationRecorder());
+	Recorder->SetSampleRateAndLength(SampleRateHz, MaxLength);
+	Recorder->bRecordLocalToWorld = bRecordInWorldSpace;
+	Recorder->SetAnimCompressionScheme(UAnimCompress_BitwiseCompressOnly::StaticClass());
+	Recorder->bAutoSaveAsset = bAutoSaveAsset;
+
+	if (InComponent)
+	{
+		CachedSkelCompForcedLodModel = InComponent->ForcedLodModel;
+		InComponent->ForcedLodModel = 1;
+	}
+}
+
+FAnimRecorderInstance::~FAnimRecorderInstance()
+{
+}
+
+bool FAnimRecorderInstance::BeginRecording()
+{
+	if (Recorder.IsValid())
+	{
+		if(Sequence.IsValid())
+		{
+			Recorder->StartRecord(SkelComp.Get(), Sequence.Get());
+			return true;
+		}
+		else
+		{
+			return Recorder->TriggerRecordAnimation(SkelComp.Get(), AssetPath, AssetName);
+		}
+	}
+
+	return false;
+}
+
+void FAnimRecorderInstance::Update(float DeltaTime)
+{
+	if (Recorder.IsValid())
+	{
+		Recorder->UpdateRecord(SkelComp.Get(), DeltaTime);
+	}
+}
+
+void FAnimRecorderInstance::FinishRecording()
+{
+	const FText FinishRecordingAnimationSlowTask = LOCTEXT("FinishRecordingAnimationSlowTask", "Finalizing recorded animation");
+	if (Recorder.IsValid())
+	{
+		Recorder->StopRecord(true);
+	}
+
+	// restore force lod setting
+	if (SkelComp.IsValid())
+	{
+		SkelComp->ForcedLodModel = CachedSkelCompForcedLodModel;
+	}
+}
+
+bool FAnimationRecorderManager::RecordAnimation(USkeletalMeshComponent* Component, const FString& AssetPath, const FString& AssetName, const FAnimationRecordingSettings& Settings)
+{
+	if (Component)
+	{
+		FAnimRecorderInstance NewInst;
+		NewInst.Init(Component, AssetPath, AssetName, Settings.SampleRate, Settings.Length, Settings.bRecordInWorldSpace, Settings.bAutoSaveAsset);
+		bool const bSuccess = NewInst.BeginRecording();
+		if (bSuccess)
+		{
+			RecorderInstances.Add(NewInst);
+		}
+
+	#if WITH_EDITOR
+			// if recording via PIE, be sure to stop recording cleanly when PIE ends
+			UWorld const* const World = Component->GetWorld();
+			if (World && World->IsPlayInEditor())
+			{
+				FEditorDelegates::EndPIE.AddRaw(this, &FAnimationRecorderManager::HandleEndPIE);
+			}
+	#endif
+
+		return bSuccess;
+	}
+
+	return false;
+}
+
+bool FAnimationRecorderManager::RecordAnimation(USkeletalMeshComponent* Component, UAnimSequence* Sequence, const FAnimationRecordingSettings& Settings)
+{
+	if (Component)
+	{
+		FAnimRecorderInstance NewInst;
+		NewInst.Init(Component, Sequence, Settings.SampleRate, Settings.Length, Settings.bRecordInWorldSpace, Settings.bAutoSaveAsset);
+		bool const bSuccess = NewInst.BeginRecording();
+		if (bSuccess)
+		{
+			RecorderInstances.Add(NewInst);
+		}
+
+	#if WITH_EDITOR
+			// if recording via PIE, be sure to stop recording cleanly when PIE ends
+			UWorld const* const World = Component->GetWorld();
+			if (World && World->IsPlayInEditor())
+			{
+				FEditorDelegates::EndPIE.AddRaw(this, &FAnimationRecorderManager::HandleEndPIE);
+			}
+	#endif
+
+		return bSuccess;
+	}
+
+	return false;
+}
+
+void FAnimationRecorderManager::HandleEndPIE(bool bSimulating)
+{
+	StopRecordingAllAnimations();
+}
+
+bool FAnimationRecorderManager::IsRecording(USkeletalMeshComponent* Component)
+{
+	for (FAnimRecorderInstance& Instance : RecorderInstances)
+	{
+		if (Instance.SkelComp == Component)
+		{
+			return Instance.Recorder->InRecording();
+		}
+	}
+
+	return false;
+}
+
+bool FAnimationRecorderManager::IsRecording()
+{
+	for (FAnimRecorderInstance& Instance : RecorderInstances)
+	{
+		if (Instance.Recorder->InRecording())
+		{
+			return true;
+		}
+	}
+
+	return false;
+}
+
+UAnimSequence* FAnimationRecorderManager::GetCurrentlyRecordingSequence(USkeletalMeshComponent* Component)
+{
+	for (FAnimRecorderInstance& Instance : RecorderInstances)
+	{
+		if (Instance.SkelComp == Component)
+		{
+			return Instance.Recorder->GetAnimationObject();
+		}
+	}
+
+	return nullptr;
+}
+
+float FAnimationRecorderManager::GetCurrentRecordingTime(USkeletalMeshComponent* Component)
+{
+	for (FAnimRecorderInstance& Instance : RecorderInstances)
+	{
+		if (Instance.SkelComp == Component)
+		{
+			return Instance.Recorder->GetTimeRecorded();
+		}
+	}
+
+	return 0.0f;
+}
+
+void FAnimationRecorderManager::StopRecordingAnimation(USkeletalMeshComponent* Component)
+{
+	for (int32 Idx = 0; Idx < RecorderInstances.Num(); ++Idx)
+	{
+		FAnimRecorderInstance& Inst = RecorderInstances[Idx];
+		if (Inst.SkelComp == Component)
+		{
+			// stop and finalize recoded data
+			Inst.FinishRecording();
+
+			// remove instance, which will clean itself up
+			RecorderInstances.RemoveAtSwap(Idx);
+
+			// all done
+			break;
+		}
+	}
+}
+
+void FAnimationRecorderManager::StopRecordingDeadAnimations()
+{
+	RecorderInstances.RemoveAll([](FAnimRecorderInstance& Instance)
+		{
+			if (!Instance.SkelComp.IsValid())
+			{
+				// stop and finalize recorded data
+				Instance.FinishRecording();
+
+				// make sure we are cleaned up
+				return true;
+			}
+
+			return false;
+		}
+	);
+}
+
+void FAnimationRecorderManager::StopRecordingAllAnimations()
+{
+	for (auto& Inst : RecorderInstances)
+	{
+		Inst.FinishRecording();
+	}
+
+	RecorderInstances.Empty();
+}
+
+#undef LOCTEXT_NAMESPACE
+
diff --git a/Engine/Plugins/Editor/SequenceRecorder/Source/SequenceRecorder/Private/AnimationRecorder.h b/Engine/Plugins/Editor/SequenceRecorder/Source/SequenceRecorder/Private/AnimationRecorder.h
new file mode 100644
--- /dev/null
+++ b/Engine/Plugins/Editor/SequenceRecorder/Source/SequenceRecorder/Private/AnimationRecorder.h
@@ -0,0 +1,137 @@
+// Copyright 1998-2016 Epic Games, Inc. All Rights Reserved.
+
+#pragma once
+
+//////////////////////////////////////////////////////////////////////////
+// FAnimationRecorder
+
+// records the mesh pose to animation input
+struct FAnimationRecorder
+{
+private:
+	/** Frame count used to signal an unbounded animation */
+	static const int32 UnBoundedFrameCount = -1;
+
+private:
+	float IntervalTime;
+	int32 MaxFrame;
+	int32 LastFrame;
+	float TimePassed;
+	UAnimSequence* AnimationObject;
+	TArray<FTransform> PreviousSpacesBases;
+	FTransform PreviousComponentToWorld;
+
+	/** Array of currently active notifies that have duration */
+	TArray<TPair<const FAnimNotifyEvent*, bool>> ActiveNotifies;
+
+	/** Unique notifies added to this sequence during recording */
+	TMap<UAnimNotify*, UAnimNotify*> UniqueNotifies;
+
+	/** Unique notify states added to this sequence during recording */
+	TMap<UAnimNotifyState*, UAnimNotifyState*> UniqueNotifyStates;
+
+	static float DefaultSampleRate;
+
+public:
+	FAnimationRecorder();
+	~FAnimationRecorder();
+
+	/** Starts recording an animation. Prompts for asset path and name via dialog if none provided */
+	bool TriggerRecordAnimation(USkeletalMeshComponent* Component);
+	bool TriggerRecordAnimation(USkeletalMeshComponent* Component, const FString& AssetPath, const FString& AssetName);
+
+	void StartRecord(USkeletalMeshComponent* Component, UAnimSequence* InAnimationObject);
+	UAnimSequence* StopRecord(bool bShowMessage);
+	void UpdateRecord(USkeletalMeshComponent* Component, float DeltaTime);
+	UAnimSequence* GetAnimationObject() const { return AnimationObject; }
+	bool InRecording() const { return AnimationObject != nullptr; }
+	float GetTimeRecorded() const { return TimePassed; }
+
+	/** Sets a new sample rate & max length for this recorder. Don't call while recording. */
+	void SetSampleRateAndLength(float SampleRateHz, float LengthInMinutes);
+
+	bool SetAnimCompressionScheme(TSubclassOf<class UAnimCompress> SchemeClass);
+
+	/** If true, it will record root to include LocalToWorld */
+	uint8 bRecordLocalToWorld :1;
+	/** If true, asset will be saved to disk after recording. If false, asset will remain in mem and can be manually saved. */
+	uint8 bAutoSaveAsset : 1;
+
+private:
+	void Record(USkeletalMeshComponent* Component, FTransform const& ComponentToWorld, TArray<FTransform> SpacesBases, int32 FrameToAdd);
+
+	void RecordNotifies(USkeletalMeshComponent* Component, const TArray<const struct FAnimNotifyEvent*>& AnimNotifies, float DeltaTime, float RecordTime);
+
+	void FixupNotifies();
+};
+
+//////////////////////////////////////////////////////////////////////////
+// FAnimRecorderInstance
+
+struct FAnimRecorderInstance
+{
+public:
+	FAnimRecorderInstance();
+	~FAnimRecorderInstance();
+
+	void Init(USkeletalMeshComponent* InComponent, const FString& InAssetPath, const FString& InAssetName, float SampleRateHz, float MaxLength, bool bRecordInWorldSpace, bool bAutoSaveAsset);
+
+	void Init(USkeletalMeshComponent* InComponent, UAnimSequence* InSequence, float SampleRateHz, float MaxLength, bool bRecordInWorldSpace, bool bAutoSaveAsset);
+
+	bool BeginRecording();
+	void Update(float DeltaTime);
+	void FinishRecording();
+
+	TWeakObjectPtr<USkeletalMeshComponent> SkelComp;
+	TWeakObjectPtr<UAnimSequence> Sequence;
+	FString AssetPath;
+	FString AssetName;
+
+	/** Original ForcedLodModel setting on the SkelComp, so we can modify it and restore it when we are done. */
+	int CachedSkelCompForcedLodModel;
+
+	TSharedPtr<FAnimationRecorder> Recorder;
+};
+
+
+//////////////////////////////////////////////////////////////////////////
+// FAnimationRecorderManager
+
+struct FAnimationRecorderManager
+{
+public:
+	/** Singleton accessor */
+	static FAnimationRecorderManager& Get();
+
+	/** Destructor */
+	virtual ~FAnimationRecorderManager();
+
+	/** Starts recording an animation. */
+	bool RecordAnimation(USkeletalMeshComponent* Component, const FString& AssetPath = FString(), const FString& AssetName = FString(), const FAnimationRecordingSettings& Settings = FAnimationRecordingSettings());
+
+	bool RecordAnimation(USkeletalMeshComponent* Component, UAnimSequence* Sequence, const FAnimationRecordingSettings& Settings = FAnimationRecordingSettings());
+
+	bool IsRecording(USkeletalMeshComponent* Component);
+
+	bool IsRecording();
+
+	UAnimSequence* GetCurrentlyRecordingSequence(USkeletalMeshComponent* Component);
+	float GetCurrentRecordingTime(USkeletalMeshComponent* Component);
+	void StopRecordingAnimation(USkeletalMeshComponent* Component);
+	void StopRecordingAllAnimations();
+
+	void Tick(float DeltaTime);
+
+	void Tick(USkeletalMeshComponent* Component, float DeltaTime);
+
+	void StopRecordingDeadAnimations();
+
+private:
+	/** Constructor, private - use Get() function */
+	FAnimationRecorderManager();
+
+	TArray<FAnimRecorderInstance> RecorderInstances;
+
+	void HandleEndPIE(bool bSimulating);
+};
+
diff --git a/Engine/Plugins/Editor/SequenceRecorder/Source/SequenceRecorder/Private/IMovieScenePropertyRecorder.h b/Engine/Plugins/Editor/SequenceRecorder/Source/SequenceRecorder/Private/IMovieScenePropertyRecorder.h
new file mode 100644
--- /dev/null
+++ b/Engine/Plugins/Editor/SequenceRecorder/Source/SequenceRecorder/Private/IMovieScenePropertyRecorder.h
@@ -0,0 +1,11 @@
+// Copyright 1998-2015 Epic Games, Inc. All Rights Reserved.
+
+#pragma once
+
+class IMovieScenePropertyRecorder
+{
+public:
+	virtual void CreateSection(class AActor* SourceActor, class UMovieScene* MovieScene, const FGuid& Guid, bool bRecord) = 0;
+	virtual void FinalizeSection() = 0;
+	virtual void Record(float CurrentTime) = 0;
+};
\ No newline at end of file
diff --git a/Engine/Plugins/Editor/SequenceRecorder/Source/SequenceRecorder/Private/PropertyRecorders/MovieScene3DTransformPropertyRecorder.cpp b/Engine/Plugins/Editor/SequenceRecorder/Source/SequenceRecorder/Private/PropertyRecorders/MovieScene3DTransformPropertyRecorder.cpp
new file mode 100644
--- /dev/null
+++ b/Engine/Plugins/Editor/SequenceRecorder/Source/SequenceRecorder/Private/PropertyRecorders/MovieScene3DTransformPropertyRecorder.cpp
@@ -0,0 +1,98 @@
+// Copyright 1998-2015 Epic Games, Inc. All Rights Reserved.
+
+#include "SequenceRecorderPrivatePCH.h"
+#include "MovieScene3DTransformPropertyRecorder.h"
+#include "MovieScene.h"
+#include "MovieScene3DTransformSection.h"
+#include "MovieScene3DTransformTrack.h"
+
+static const FName TransformTrackName = TEXT("Transform");
+
+void FMovieScene3DTransformPropertyRecorder::CreateSection(AActor* SourceActor, UMovieScene* MovieScene, const FGuid& Guid, bool bRecord) 
+{
+	ActorToRecord = SourceActor;
+
+	UMovieScene3DTransformTrack* TransformTrack = MovieScene->AddTrack<UMovieScene3DTransformTrack>(Guid);
+	if(TransformTrack)
+	{
+		TransformTrack->SetPropertyNameAndPath(TransformTrackName, TransformTrackName.ToString());
+
+		MovieSceneSection = Cast<UMovieScene3DTransformSection>(TransformTrack->CreateNewSection());
+
+		TransformTrack->AddSection(*MovieSceneSection);
+
+		const bool bUnwindRotation = false;
+
+		MovieSceneSection->SetDefault(FTransformKey(EKey3DTransformChannel::Translation, EAxis::X, 0.0f, bUnwindRotation));
+		MovieSceneSection->SetDefault(FTransformKey(EKey3DTransformChannel::Translation, EAxis::Y, 0.0f, bUnwindRotation));
+		MovieSceneSection->SetDefault(FTransformKey(EKey3DTransformChannel::Translation, EAxis::Z, 0.0f, bUnwindRotation));
+
+		MovieSceneSection->SetDefault(FTransformKey(EKey3DTransformChannel::Rotation, EAxis::X, 0.0f, bUnwindRotation));
+		MovieSceneSection->SetDefault(FTransformKey(EKey3DTransformChannel::Rotation, EAxis::Y, 0.0f, bUnwindRotation));
+		MovieSceneSection->SetDefault(FTransformKey(EKey3DTransformChannel::Rotation, EAxis::Z, 0.0f, bUnwindRotation));
+
+		MovieSceneSection->SetDefault(FTransformKey(EKey3DTransformChannel::Scale, EAxis::X, 1.0f, bUnwindRotation));
+		MovieSceneSection->SetDefault(FTransformKey(EKey3DTransformChannel::Scale, EAxis::Y, 1.0f, bUnwindRotation));
+		MovieSceneSection->SetDefault(FTransformKey(EKey3DTransformChannel::Scale, EAxis::Z, 1.0f, bUnwindRotation));
+
+		MovieSceneSection->SetStartTime(0.0f);
+	}
+
+	bRecording = bRecord;
+}
+
+void FMovieScene3DTransformPropertyRecorder::FinalizeSection()
+{
+	bRecording = false;
+
+	// now remove linear keys
+	TPair<FRichCurve*, float> CurvesAndTolerances[] =
+	{
+		TPairInitializer<FRichCurve*, float>(&MovieSceneSection->GetTranslationCurve(EAxis::X), 0.1f),
+		TPairInitializer<FRichCurve*, float>(&MovieSceneSection->GetTranslationCurve(EAxis::Y), 0.1f),
+		TPairInitializer<FRichCurve*, float>(&MovieSceneSection->GetTranslationCurve(EAxis::Z), 0.1f),
+		TPairInitializer<FRichCurve*, float>(&MovieSceneSection->GetRotationCurve(EAxis::X), 0.25f),
+		TPairInitializer<FRichCurve*, float>(&MovieSceneSection->GetRotationCurve(EAxis::Y), 0.25f),
+		TPairInitializer<FRichCurve*, float>(&MovieSceneSection->GetRotationCurve(EAxis::Z), 0.25f),
+		TPairInitializer<FRichCurve*, float>(&MovieSceneSection->GetScaleCurve(EAxis::X), 0.1f),
+		TPairInitializer<FRichCurve*, float>(&MovieSceneSection->GetScaleCurve(EAxis::Y), 0.1f),
+		TPairInitializer<FRichCurve*, float>(&MovieSceneSection->GetScaleCurve(EAxis::Z), 0.1f),
+	};
+
+	for(TPair<FRichCurve*, float>& CurveAndTolerance : CurvesAndTolerances)
+	{
+		CurveAndTolerance.Key->RemoveRedundantKeys(CurveAndTolerance.Value);
+	}
+}
+
+void FMovieScene3DTransformPropertyRecorder::Record(float CurrentTime)
+{
+	if(ActorToRecord.IsValid())
+	{
+		MovieSceneSection->SetEndTime(CurrentTime);
+
+		if(bRecording)
+		{
+			const bool bUnwindRotation = false;
+
+			USkeletalMeshComponent* SkeletalMeshComponent = ActorToRecord->FindComponentByClass<USkeletalMeshComponent>();
+
+			const FTransform Transform = SkeletalMeshComponent ? SkeletalMeshComponent->ComponentToWorld : ActorToRecord->GetTransform();
+			const FVector Translation = Transform.GetTranslation();
+			const FVector Rotation = Transform.GetRotation().Rotator().Euler();
+			const FVector Scale = Transform.GetScale3D();
+
+			MovieSceneSection->AddKey(CurrentTime, FTransformKey(EKey3DTransformChannel::Translation, EAxis::X, Translation.X, bUnwindRotation), EMovieSceneKeyInterpolation::Auto);
+			MovieSceneSection->AddKey(CurrentTime, FTransformKey(EKey3DTransformChannel::Translation, EAxis::Y, Translation.Y, bUnwindRotation), EMovieSceneKeyInterpolation::Auto);
+			MovieSceneSection->AddKey(CurrentTime, FTransformKey(EKey3DTransformChannel::Translation, EAxis::Z, Translation.Z, bUnwindRotation), EMovieSceneKeyInterpolation::Auto);
+
+			MovieSceneSection->AddKey(CurrentTime, FTransformKey(EKey3DTransformChannel::Rotation, EAxis::X, Rotation.X, bUnwindRotation), EMovieSceneKeyInterpolation::Auto);
+			MovieSceneSection->AddKey(CurrentTime, FTransformKey(EKey3DTransformChannel::Rotation, EAxis::Y, Rotation.Y, bUnwindRotation), EMovieSceneKeyInterpolation::Auto);
+			MovieSceneSection->AddKey(CurrentTime, FTransformKey(EKey3DTransformChannel::Rotation, EAxis::Z, Rotation.Z, bUnwindRotation), EMovieSceneKeyInterpolation::Auto);
+
+			MovieSceneSection->AddKey(CurrentTime, FTransformKey(EKey3DTransformChannel::Scale, EAxis::X, Scale.X, bUnwindRotation), EMovieSceneKeyInterpolation::Auto);
+			MovieSceneSection->AddKey(CurrentTime, FTransformKey(EKey3DTransformChannel::Scale, EAxis::Y, Scale.Y, bUnwindRotation), EMovieSceneKeyInterpolation::Auto);
+			MovieSceneSection->AddKey(CurrentTime, FTransformKey(EKey3DTransformChannel::Scale, EAxis::Z, Scale.Z, bUnwindRotation), EMovieSceneKeyInterpolation::Auto);
+		}
+	}
+}
\ No newline at end of file
diff --git a/Engine/Plugins/Editor/SequenceRecorder/Source/SequenceRecorder/Private/PropertyRecorders/MovieScene3DTransformPropertyRecorder.h b/Engine/Plugins/Editor/SequenceRecorder/Source/SequenceRecorder/Private/PropertyRecorders/MovieScene3DTransformPropertyRecorder.h
new file mode 100644
--- /dev/null
+++ b/Engine/Plugins/Editor/SequenceRecorder/Source/SequenceRecorder/Private/PropertyRecorders/MovieScene3DTransformPropertyRecorder.h
@@ -0,0 +1,22 @@
+// Copyright 1998-2015 Epic Games, Inc. All Rights Reserved.
+
+#pragma once
+
+#include "IMovieScenePropertyRecorder.h"
+
+class FMovieScene3DTransformPropertyRecorder : public IMovieScenePropertyRecorder
+{
+public:
+	virtual ~FMovieScene3DTransformPropertyRecorder() {}
+
+	virtual void CreateSection(class AActor* SourceActor, class UMovieScene* MovieScene, const FGuid& Guid, bool bRecord) override;
+	virtual void FinalizeSection() override;
+	virtual void Record(float CurrentTime) override;
+
+private:
+	TLazyObjectPtr<class AActor> ActorToRecord;
+
+	TWeakObjectPtr<class UMovieScene3DTransformSection> MovieSceneSection;
+
+	bool bRecording;
+};
\ No newline at end of file
diff --git a/Engine/Plugins/Editor/SequenceRecorder/Source/SequenceRecorder/Private/PropertyRecorders/MovieSceneVisibilityPropertyRecorder.cpp b/Engine/Plugins/Editor/SequenceRecorder/Source/SequenceRecorder/Private/PropertyRecorders/MovieSceneVisibilityPropertyRecorder.cpp
new file mode 100644
--- /dev/null
+++ b/Engine/Plugins/Editor/SequenceRecorder/Source/SequenceRecorder/Private/PropertyRecorders/MovieSceneVisibilityPropertyRecorder.cpp
@@ -0,0 +1,57 @@
+// Copyright 1998-2015 Epic Games, Inc. All Rights Reserved.
+
+#include "SequenceRecorderPrivatePCH.h"
+#include "MovieSceneVisibilityPropertyRecorder.h"
+#include "MovieSceneVisibilitySection.h"
+#include "MovieSceneVisibilityTrack.h"
+#include "MovieScene.h"
+
+static const FName VisibilityTrackName = TEXT("Visibility");
+
+void FMovieSceneVisibilityPropertyRecorder::CreateSection(AActor* SourceActor, UMovieScene* MovieScene, const FGuid& Guid, bool bRecord)
+{
+	ActorToRecord = SourceActor;
+
+	UMovieSceneVisibilityTrack* VisibilityTrack = MovieScene->AddTrack<UMovieSceneVisibilityTrack>(Guid);
+	if(VisibilityTrack)
+	{
+		VisibilityTrack->SetPropertyNameAndPath(VisibilityTrackName, VisibilityTrackName.ToString());
+
+		MovieSceneSection = Cast<UMovieSceneVisibilitySection>(VisibilityTrack->CreateNewSection());
+
+		VisibilityTrack->AddSection(*MovieSceneSection);
+
+		MovieSceneSection->SetDefault(false);
+
+		MovieSceneSection->SetStartTime(0.0f);
+	}
+
+	bRecording = bRecord;
+
+	bWasVisible = false;
+}
+
+void FMovieSceneVisibilityPropertyRecorder::FinalizeSection()
+{
+	bRecording = false;
+}
+
+void FMovieSceneVisibilityPropertyRecorder::Record(float CurrentTime)
+{
+	if(ActorToRecord.IsValid())
+	{
+		MovieSceneSection->SetEndTime(CurrentTime);
+
+		if(bRecording)
+		{
+			USkeletalMeshComponent* SkeletalMeshComponent = ActorToRecord->FindComponentByClass<USkeletalMeshComponent>();
+
+			const bool bVisible = SkeletalMeshComponent ? !SkeletalMeshComponent->bHiddenInGame : !ActorToRecord->bHidden;
+			if(bVisible != bWasVisible)
+			{
+				MovieSceneSection->AddKey(CurrentTime, bVisible, EMovieSceneKeyInterpolation::Break);
+			}
+			bWasVisible = bVisible;
+		}
+	}
+}
\ No newline at end of file
diff --git a/Engine/Plugins/Editor/SequenceRecorder/Source/SequenceRecorder/Private/PropertyRecorders/MovieSceneVisibilityPropertyRecorder.h b/Engine/Plugins/Editor/SequenceRecorder/Source/SequenceRecorder/Private/PropertyRecorders/MovieSceneVisibilityPropertyRecorder.h
new file mode 100644
--- /dev/null
+++ b/Engine/Plugins/Editor/SequenceRecorder/Source/SequenceRecorder/Private/PropertyRecorders/MovieSceneVisibilityPropertyRecorder.h
@@ -0,0 +1,24 @@
+// Copyright 1998-2015 Epic Games, Inc. All Rights Reserved.
+
+#pragma once
+
+#include "IMovieScenePropertyRecorder.h"
+
+class FMovieSceneVisibilityPropertyRecorder : public IMovieScenePropertyRecorder
+{
+public:
+	virtual ~FMovieSceneVisibilityPropertyRecorder() {}
+
+	virtual void CreateSection(class AActor* SourceActor, class UMovieScene* MovieScene, const FGuid& Guid, bool bRecord) override;
+	virtual void FinalizeSection() override;
+	virtual void Record(float CurrentTime) override;
+
+private:
+	TLazyObjectPtr<class AActor> ActorToRecord;
+
+	TWeakObjectPtr<class UMovieSceneVisibilitySection> MovieSceneSection;
+
+	bool bRecording;
+
+	bool bWasVisible;
+};
\ No newline at end of file
diff --git a/Engine/Plugins/Editor/SequenceRecorder/Source/SequenceRecorder/Private/SSequenceRecorder.cpp b/Engine/Plugins/Editor/SequenceRecorder/Source/SequenceRecorder/Private/SSequenceRecorder.cpp
new file mode 100644
--- /dev/null
+++ b/Engine/Plugins/Editor/SequenceRecorder/Source/SequenceRecorder/Private/SSequenceRecorder.cpp
@@ -0,0 +1,356 @@
+// Copyright 1998-2016 Epic Games, Inc. All Rights Reserved.
+
+#include "SequenceRecorderPrivatePCH.h"
+#include "SSequenceRecorder.h"
+#include "ActorRecording.h"
+#include "SequenceRecorderCommands.h"
+#include "SNumericEntryBox.h"
+#include "SequenceRecorderSettings.h"
+#include "SequenceRecorder.h"
+
+#define LOCTEXT_NAMESPACE "SequenceRecorder"
+
+static const FName ActorColumnName(TEXT("Actor"));
+static const FName AnimationColumnName(TEXT("Animation"));
+static const FName LengthColumnName(TEXT("Length"));
+
+/** A widget to display information about an animation recording in the list view */
+class SSequenceRecorderListRow : public SMultiColumnTableRow<UActorRecording*>
+{
+public:
+	SLATE_BEGIN_ARGS(SSequenceRecorderListRow) {}
+
+	/** The list item for this row */
+	SLATE_ARGUMENT(UActorRecording*, Recording)
+
+	SLATE_END_ARGS()
+
+	void Construct(const FArguments& Args, const TSharedRef<STableViewBase>& OwnerTableView)
+	{
+		RecordingPtr = Args._Recording;
+
+		SMultiColumnTableRow<UActorRecording*>::Construct(
+			FSuperRowType::FArguments()
+			.Padding(1.0f),
+			OwnerTableView
+			);
+	}
+
+	/** Overridden from SMultiColumnTableRow.  Generates a widget for this column of the list view. */
+	virtual TSharedRef<SWidget> GenerateWidgetForColumn(const FName& ColumnName) override
+	{
+		if (ColumnName == ActorColumnName)
+		{
+			return
+				SNew(STextBlock)
+				.Text(TAttribute<FText>::Create(TAttribute<FText>::FGetter::CreateSP(this, &SSequenceRecorderListRow::GetRecordingActorName)));
+		}
+		else if (ColumnName == AnimationColumnName)
+		{
+			return
+				SNew(STextBlock)
+				.Text(TAttribute<FText>::Create(TAttribute<FText>::FGetter::CreateSP(this, &SSequenceRecorderListRow::GetRecordingAnimationName)));
+		}
+		else if (ColumnName == LengthColumnName)
+		{
+			return
+				SNew(STextBlock)
+				.Text(TAttribute<FText>::Create(TAttribute<FText>::FGetter::CreateSP(this, &SSequenceRecorderListRow::GetRecordingLengthText)));
+		}
+
+		return SNullWidget::NullWidget;
+	}
+
+private:
+	FText GetRecordingActorName() const
+	{
+		FText ActorName(LOCTEXT("InvalidActorName", "None"));
+		if (RecordingPtr.IsValid() && RecordingPtr.Get()->ActorToRecord.IsValid())
+		{
+			ActorName = FText::FromString(RecordingPtr.Get()->ActorToRecord.Get()->GetActorLabel());
+		}
+
+		return ActorName;
+	}
+
+	FText GetRecordingAnimationName() const
+	{
+		FText AnimationName(LOCTEXT("InvalidAnimationName", "None"));
+		if (RecordingPtr.IsValid())
+		{
+			if(!RecordingPtr.Get()->bSpecifyTargetAnimation)
+			{
+				AnimationName = LOCTEXT("AutoCreatedAnimationName", "Auto-created");
+			}
+			else if(RecordingPtr.Get()->TargetAnimation.IsValid())
+			{
+				AnimationName = FText::FromString(RecordingPtr.Get()->TargetAnimation.Get()->GetName());
+			}
+		}
+
+		return AnimationName;
+	}
+
+	FText GetRecordingLengthText() const
+	{
+		FText LengthText(LOCTEXT("InvalidLengthName", "None"));
+		if (RecordingPtr.IsValid())
+		{
+			LengthText = FText::AsNumber(RecordingPtr.Get()->AnimationSettings.Length);
+		}
+
+		return LengthText;
+	}
+
+	TWeakObjectPtr<UActorRecording> RecordingPtr;
+};
+
+
+void SSequenceRecorder::Construct(const FArguments& Args)
+{
+	CommandList = MakeShareable(new FUICommandList);
+
+	BindCommands();
+
+	FPropertyEditorModule& PropertyEditorModule = FModuleManager::Get().GetModuleChecked<FPropertyEditorModule>("PropertyEditor");
+
+	FDetailsViewArgs DetailsViewArgs;
+	DetailsViewArgs.NameAreaSettings = FDetailsViewArgs::HideNameArea;
+	DetailsViewArgs.bAllowSearch = false;
+
+	ActorRecordingDetailsView = PropertyEditorModule.CreateDetailView( DetailsViewArgs );
+	SequenceRecordingDetailsView = PropertyEditorModule.CreateDetailView( DetailsViewArgs );
+	SequenceRecordingDetailsView->SetObject(GetMutableDefault<USequenceRecorderSettings>());
+
+	FToolBarBuilder ToolBarBuilder(CommandList, FMultiBoxCustomization::None);
+
+	ToolBarBuilder.BeginSection(TEXT("Recording"));
+	{
+		ToolBarBuilder.AddToolBarButton(FSequenceRecorderCommands::Get().RecordAll);
+		ToolBarBuilder.AddToolBarButton(FSequenceRecorderCommands::Get().StopAll);
+	}
+	ToolBarBuilder.EndSection();
+
+	ToolBarBuilder.BeginSection(TEXT("RecordingManagement"));
+	{
+		ToolBarBuilder.AddToolBarButton(FSequenceRecorderCommands::Get().AddRecording);
+		ToolBarBuilder.AddToolBarButton(FSequenceRecorderCommands::Get().RemoveRecording);
+	}
+	ToolBarBuilder.EndSection();
+
+	ChildSlot
+	[
+		SNew(SSplitter)
+		.Orientation(Orient_Vertical)
+		+SSplitter::Slot()
+		.Value(0.33f)
+		[
+			SNew(SVerticalBox)
+			+SVerticalBox::Slot()
+			.AutoHeight()
+			.Padding(FMargin(0.0f, 4.0f, 0.0f, 0.0f))
+			[
+				ToolBarBuilder.MakeWidget()
+			]
+			+SVerticalBox::Slot()
+			.FillHeight(1.0f)
+			.Padding(FMargin(0.0f, 4.0f, 0.0f, 0.0f))
+			[
+				SNew(SBorder)
+				.BorderImage(FEditorStyle::GetBrush("ToolPanel.GroupBorder"))
+				.Padding(FMargin(4.0f, 4.0f))
+				[
+					SNew(SOverlay)
+					+SOverlay::Slot()
+					[
+						SAssignNew(ListView, SListView<UActorRecording*>)
+						.ListItemsSource(&FSequenceRecorder::Get().GetQueuedRecordings())
+						.SelectionMode(ESelectionMode::SingleToggle)
+						.OnGenerateRow(this, &SSequenceRecorder::MakeListViewWidget)
+						.OnSelectionChanged(this, &SSequenceRecorder::OnSelectionChanged)
+						.HeaderRow
+						(
+						SNew(SHeaderRow)
+						+ SHeaderRow::Column(ActorColumnName)
+						.FillWidth(43.0f)
+						.DefaultLabel(LOCTEXT("ActorHeaderName", "Actor"))
+						+ SHeaderRow::Column(AnimationColumnName)
+						.FillWidth(43.0f)
+						.DefaultLabel(LOCTEXT("AnimationHeaderName", "Animation"))
+						+ SHeaderRow::Column(LengthColumnName)
+						.FillWidth(14.0f)
+						.DefaultLabel(LOCTEXT("LengthHeaderName", "Length"))
+						)
+					]
+					+SOverlay::Slot()
+					[
+						SNew(SVerticalBox)
+						+SVerticalBox::Slot()
+						.VAlign(VAlign_Bottom)
+						.MaxHeight(2.0f)
+						[
+							SAssignNew(DelayProgressBar, SProgressBar)
+							.Percent(this, &SSequenceRecorder::GetDelayPercent)
+							.Visibility(this, &SSequenceRecorder::GetDelayProgressVisibilty)
+						]
+					]
+				]
+			]
+		]
+		+SSplitter::Slot()
+		.Value(0.66f)
+		[
+			SNew(SVerticalBox)
+			+SVerticalBox::Slot()
+			.AutoHeight()
+			[
+				SequenceRecordingDetailsView.ToSharedRef()
+			]
+			+SVerticalBox::Slot()
+			.AutoHeight()
+			[
+				ActorRecordingDetailsView.ToSharedRef()
+			]
+		]
+	];
+
+	// Register refresh timer
+	if (!ActiveTimerHandle.IsValid())
+	{
+		ActiveTimerHandle = RegisterActiveTimer(0.1f, FWidgetActiveTimerDelegate::CreateSP(this, &SSequenceRecorder::HandleRefreshItems));
+	}
+}
+
+void SSequenceRecorder::BindCommands()
+{
+	CommandList->MapAction(FSequenceRecorderCommands::Get().RecordAll,
+		FExecuteAction::CreateSP(this, &SSequenceRecorder::HandleRecord),
+		FCanExecuteAction::CreateSP(this, &SSequenceRecorder::CanRecord),
+		FIsActionChecked(),
+		FIsActionButtonVisible::CreateSP(this, &SSequenceRecorder::IsRecordVisible)
+		);
+
+	CommandList->MapAction(FSequenceRecorderCommands::Get().StopAll,
+		FExecuteAction::CreateSP(this, &SSequenceRecorder::HandleStopAll),
+		FCanExecuteAction::CreateSP(this, &SSequenceRecorder::CanStopAll),
+		FIsActionChecked(),
+		FIsActionButtonVisible::CreateSP(this, &SSequenceRecorder::IsStopAllVisible)
+		);
+
+	CommandList->MapAction(FSequenceRecorderCommands::Get().AddRecording,
+		FExecuteAction::CreateSP(this, &SSequenceRecorder::HandleAddRecording),
+		FCanExecuteAction::CreateSP(this, &SSequenceRecorder::CanAddRecording)
+		);
+
+	CommandList->MapAction(FSequenceRecorderCommands::Get().RemoveRecording,
+		FExecuteAction::CreateSP(this, &SSequenceRecorder::HandleRemoveRecording),
+		FCanExecuteAction::CreateSP(this, &SSequenceRecorder::CanRemoveRecording)
+		);
+}
+
+TSharedRef<ITableRow> SSequenceRecorder::MakeListViewWidget(UActorRecording* Recording, const TSharedRef<STableViewBase>& OwnerTable) const
+{
+	return
+		SNew(SSequenceRecorderListRow, OwnerTable)
+		.Recording(Recording);
+}
+
+void SSequenceRecorder::OnSelectionChanged(UActorRecording* Recording, ESelectInfo::Type SelectionType) const
+{
+	if(Recording)
+	{
+		ActorRecordingDetailsView->SetObject(Recording);
+	}
+	else
+	{
+		ActorRecordingDetailsView->SetObject(nullptr);
+	}
+}
+
+void SSequenceRecorder::HandleRecord()
+{
+	FSequenceRecorder::Get().StartRecording();
+}
+
+bool SSequenceRecorder::CanRecord() const
+{
+	return FSequenceRecorder::Get().HasQueuedRecordings();
+}
+
+bool SSequenceRecorder::IsRecordVisible() const
+{
+	return !FAnimationRecorderManager::Get().IsRecording() && !FSequenceRecorder::Get().IsDelaying();
+}
+
+void SSequenceRecorder::HandleStopAll()
+{
+	FSequenceRecorder::Get().StopRecording();
+}
+
+bool SSequenceRecorder::CanStopAll() const
+{
+	return FAnimationRecorderManager::Get().IsRecording() || FSequenceRecorder::Get().IsDelaying();
+}
+
+bool SSequenceRecorder::IsStopAllVisible() const
+{
+	return FAnimationRecorderManager::Get().IsRecording() || FSequenceRecorder::Get().IsDelaying();
+}
+
+void SSequenceRecorder::HandleAddRecording()
+{
+	FSequenceRecorder::Get().AddNewQueuedRecording();
+}
+
+bool SSequenceRecorder::CanAddRecording() const
+{
+	return !FAnimationRecorderManager::Get().IsRecording();
+}
+
+void SSequenceRecorder::HandleRemoveRecording()
+{
+	TArray<UActorRecording*> SelectedRecordings;
+	ListView->GetSelectedItems(SelectedRecordings);
+	UActorRecording* SelectedRecording = SelectedRecordings.Num() > 0 ? SelectedRecordings[0] : nullptr;
+	if(SelectedRecording)
+	{
+		FSequenceRecorder::Get().RemoveQueuedRecording(SelectedRecording);
+
+		TArray<TWeakObjectPtr<UObject>> SelectedObjects = ActorRecordingDetailsView->GetSelectedObjects();
+		if(SelectedObjects.Num() > 0 && SelectedObjects[0].Get() == SelectedRecording)
+		{
+			ActorRecordingDetailsView->SetObject(nullptr);
+		}
+	}
+}
+
+bool SSequenceRecorder::CanRemoveRecording() const
+{
+	return ListView->GetNumItemsSelected() > 0 && !FAnimationRecorderManager::Get().IsRecording();
+}
+
+EActiveTimerReturnType SSequenceRecorder::HandleRefreshItems(double InCurrentTime, float InDeltaTime)
+{
+	if(FSequenceRecorder::Get().AreQueuedRecordingsDirty())
+	{
+		ListView->RequestListRefresh();
+
+		FSequenceRecorder::Get().ResetQueuedRecordingsDirty();
+	}
+
+	return EActiveTimerReturnType::Continue;
+}
+
+TOptional<float> SSequenceRecorder::GetDelayPercent() const
+{
+	const float Delay = GetDefault<USequenceRecorderSettings>()->RecordingDelay;
+	const float Countdown = FSequenceRecorder::Get().GetCurrentDelay();
+	return Delay > 0.0f ? Countdown / Delay : 0.0f;
+}
+
+EVisibility SSequenceRecorder::GetDelayProgressVisibilty() const
+{
+	return FSequenceRecorder::Get().IsDelaying() ? EVisibility::Visible : EVisibility::Hidden;
+}
+
+#undef LOCTEXT_NAMESPACE
\ No newline at end of file
diff --git a/Engine/Plugins/Editor/SequenceRecorder/Source/SequenceRecorder/Private/SSequenceRecorder.h b/Engine/Plugins/Editor/SequenceRecorder/Source/SequenceRecorder/Private/SSequenceRecorder.h
new file mode 100644
--- /dev/null
+++ b/Engine/Plugins/Editor/SequenceRecorder/Source/SequenceRecorder/Private/SSequenceRecorder.h
@@ -0,0 +1,70 @@
+// Copyright 1998-2016 Epic Games, Inc. All Rights Reserved.
+
+#pragma once
+
+#include "ActorRecording.h"
+
+class SSequenceRecorder : public SCompoundWidget
+{
+	SLATE_BEGIN_ARGS(SSequenceRecorder)
+	{
+	}
+
+	SLATE_END_ARGS()
+
+	void Construct(const FArguments& Args);
+
+private:
+
+	void BindCommands();
+
+	TSharedRef<ITableRow> MakeListViewWidget(class UActorRecording* Recording, const TSharedRef<STableViewBase>& OwnerTable) const;
+
+	FText GetRecordingActorName(class UActorRecording* Recording) const;
+
+	void OnSelectionChanged(UActorRecording* Recording, ESelectInfo::Type SelectionType) const;
+
+	void HandleRecord();
+
+	EActiveTimerReturnType StartDelayedRecord(double InCurrentTime, float InDeltaTime);
+
+	bool CanRecord() const;
+
+	bool IsRecordVisible() const;
+
+	void HandleStopAll();
+
+	bool CanStopAll() const;
+
+	bool IsStopAllVisible() const;
+
+	void HandleAddRecording();
+
+	bool CanAddRecording() const;
+
+	void HandleRemoveRecording();
+
+	bool CanRemoveRecording() const;
+
+	EActiveTimerReturnType HandleRefreshItems(double InCurrentTime, float InDeltaTime);
+
+	TOptional<float> GetDelayPercent() const;
+
+	void OnDelayChanged(float NewValue);
+
+	EVisibility GetDelayProgressVisibilty() const;
+
+private:
+	TSharedPtr<IDetailsView> SequenceRecordingDetailsView;
+
+	TSharedPtr<IDetailsView> ActorRecordingDetailsView;
+
+	TSharedPtr<SListView<UActorRecording*>> ListView;
+
+	TSharedPtr<FUICommandList> CommandList;
+
+	/** The handle to the refresh tick timer */
+	TWeakPtr<FActiveTimerHandle> ActiveTimerHandle;
+
+	TSharedPtr<SProgressBar> DelayProgressBar;
+};
diff --git a/Engine/Plugins/Editor/SequenceRecorder/Source/SequenceRecorder/Private/SequenceRecorder.cpp b/Engine/Plugins/Editor/SequenceRecorder/Source/SequenceRecorder/Private/SequenceRecorder.cpp
new file mode 100644
--- /dev/null
+++ b/Engine/Plugins/Editor/SequenceRecorder/Source/SequenceRecorder/Private/SequenceRecorder.cpp
@@ -0,0 +1,384 @@
+// Copyright 1998-2016 Epic Games, Inc. All Rights Reserved.
+
+#include "SequenceRecorderPrivatePCH.h"
+#include "SequenceRecorder.h"
+#include "ActorRecording.h"
+#include "SequenceRecorderUtils.h"
+#include "SequenceRecorderSettings.h"
+#include "ObjectTools.h"
+#include "LevelSequence.h"
+#include "AssetSelection.h"
+#include "MovieSceneSkeletalAnimationTrack.h"
+#include "MovieSceneSkeletalAnimationSection.h"
+#include "MovieSceneSpawnTrack.h"
+#include "MovieScene3DTransformTrack.h"
+#include "MovieScene3DTransformSection.h"
+#include "MovieSceneBoolSection.h"
+#include "ScopedTransaction.h"
+#include "AssetRegistryModule.h"
+#include "IMovieScenePropertyRecorder.h"
+#include "MovieScene3DTransformPropertyRecorder.h"
+#include "MovieSceneVisibilityPropertyRecorder.h"
+
+FSequenceRecorder::FSequenceRecorder()
+	: bQueuedRecordingsDirty(false)
+	, CurrentDelay(0.0f)
+	, CurrentTime(0.0f)
+{
+}
+
+FSequenceRecorder& FSequenceRecorder::Get()
+{
+	static FSequenceRecorder SequenceRecorder;
+	return SequenceRecorder;
+}
+
+
+bool FSequenceRecorder::IsRecordingQueued(AActor* Actor) const
+{
+	for (UActorRecording* QueuedRecording : QueuedRecordings)
+	{
+		if (QueuedRecording->ActorToRecord == Actor)
+		{
+			return true;
+		}
+	}
+
+	return false;
+}
+
+void FSequenceRecorder::StartAllQueuedRecordings()
+{
+	for (UActorRecording* QueuedRecording : QueuedRecordings)
+	{
+		QueuedRecording->StartRecording();
+	}
+}
+
+void FSequenceRecorder::StopAllQueuedRecordings()
+{
+	for (UActorRecording* QueuedRecording : QueuedRecordings)
+	{
+		QueuedRecording->StopRecording();
+	}
+}
+
+void FSequenceRecorder::AddNewQueuedRecording(AActor* Actor, UAnimSequence* AnimSequence, float Length)
+{
+	const USequenceRecorderSettings* Settings = GetDefault<USequenceRecorderSettings>();
+
+	UActorRecording* ActorRecording = NewObject<UActorRecording>();
+	ActorRecording->AddToRoot();
+	ActorRecording->ActorToRecord = Actor;
+	ActorRecording->TargetAnimation = AnimSequence;
+	ActorRecording->AnimationSettings.Length = Length;
+	ActorRecording->AnimationSettings.bRecordInWorldSpace = !Settings->bRecordTransformsAsSeperateTrack;
+	ActorRecording->ActorSettings.bRecordTransforms = Settings->bRecordTransformsAsSeperateTrack;
+
+	QueuedRecordings.Add(ActorRecording);
+
+	bQueuedRecordingsDirty = true;
+}
+
+void FSequenceRecorder::RemoveQueuedRecording(AActor* Actor)
+{
+	for (UActorRecording* QueuedRecording : QueuedRecordings)
+	{
+		if (QueuedRecording->ActorToRecord == Actor)
+		{
+			QueuedRecording->RemoveFromRoot();
+			QueuedRecordings.Remove(QueuedRecording);
+			break;
+		}
+	}
+
+	bQueuedRecordingsDirty = true;
+}
+
+void FSequenceRecorder::RemoveQueuedRecording(class UActorRecording* Recording)
+{
+	for (UActorRecording* QueuedRecording : QueuedRecordings)
+	{
+		if (QueuedRecording == Recording)
+		{
+			QueuedRecording->RemoveFromRoot();
+			QueuedRecordings.Remove(QueuedRecording);
+			break;
+		}
+	}
+
+	bQueuedRecordingsDirty = true;
+}
+
+bool FSequenceRecorder::HasQueuedRecordings() const
+{
+	return QueuedRecordings.Num() > 0;
+}
+
+
+void FSequenceRecorder::Tick(float DeltaSeconds)
+{
+	FAnimationRecorderManager::Get().Tick(DeltaSeconds);
+
+	for(UActorRecording* Recording : QueuedRecordings)
+	{
+		Recording->Tick(DeltaSeconds);
+	}
+
+	if(CurrentDelay > 0.0f)
+	{
+		CurrentDelay -= DeltaSeconds;
+		if(CurrentDelay <= 0.0f)
+		{
+			CurrentDelay = 0.0f;
+			StartRecordingInternal();
+		}
+	}
+
+	const USequenceRecorderSettings* Settings = GetDefault<USequenceRecorderSettings>();
+	if(Settings->bCreateLevelSequence && CurrentSequence.IsValid())
+	{
+		// record our property recorders
+		for(auto& PropertyRecorder : PropertyRecorders)
+		{
+			PropertyRecorder->Record(CurrentTime);
+		}
+
+		CurrentTime += DeltaSeconds;
+
+		// check if all our actor recordings are finished
+		if(QueuedRecordings.Num() > 0)
+		{
+			bool bAllFinished = true;
+			for(UActorRecording* Recording : QueuedRecordings)
+			{
+				if(Recording->IsRecording())
+				{
+					bAllFinished = false;
+					break;
+				}
+			}
+
+			if(bAllFinished)
+			{
+				StopRecording();
+			}
+		}
+	}
+}
+
+bool FSequenceRecorder::StartRecording()
+{
+	const USequenceRecorderSettings* Settings = GetDefault<USequenceRecorderSettings>();
+
+	if(Settings->RecordingDelay > 0.0f)
+	{
+		CurrentDelay = Settings->RecordingDelay;
+		return QueuedRecordings.Num() > 0;
+	}
+	
+	return StartRecordingInternal();
+}
+
+bool FSequenceRecorder::StartRecordingInternal()
+{
+	CurrentTime = 0.0f;
+
+	const USequenceRecorderSettings* Settings = GetDefault<USequenceRecorderSettings>();
+	if(QueuedRecordings.Num() > 0)
+	{
+		ULevelSequence* LevelSequence = nullptr;
+		
+		if(Settings->bCreateLevelSequence)
+		{
+			// build an asset path
+			FString AssetPath(TEXT("/Game"));
+			AssetPath /= Settings->SequenceRecordingBasePath.Path;
+
+			FString AssetName = Settings->SequenceName.Len() > 0 ? Settings->SequenceName : TEXT("RecordedSequence");
+
+			LevelSequence = SequenceRecorderUtils::MakeNewAsset<ULevelSequence>(AssetPath, AssetName);
+			if(LevelSequence)
+			{
+				LevelSequence->Initialize();
+				CurrentSequence = LevelSequence;
+			}
+		}
+
+		// start recording 
+		bool bStartedRecordingAllActors = true;
+		for(UActorRecording* Recording : QueuedRecordings)
+		{
+			if(!Recording->StartRecording())
+			{
+				bStartedRecordingAllActors = false;
+				break;
+			}
+
+			// set up our spawnable for this actor
+			UMovieScene* MovieScene = LevelSequence->GetMovieScene();
+			UAnimSequence* AnimSequence = Recording->LastRecordedAnimation.Get();
+			UActorFactory* FactoryToUse = FActorFactoryAssetProxy::GetFactoryForAssetObject(AnimSequence);
+			check(FactoryToUse);
+
+			FString BlueprintName = AnimSequence->GetName();
+			auto DuplName = [&](FMovieSceneSpawnable& InSpawnable)
+			{
+				return InSpawnable.GetName() == BlueprintName;
+			};
+
+			int32 Index = 2;
+			FString UniqueString;
+			while (MovieScene->FindSpawnable(DuplName))
+			{
+				BlueprintName.RemoveFromEnd(UniqueString);
+				UniqueString = FString::Printf(TEXT(" (%d)"), Index);
+				BlueprintName += UniqueString;
+			}
+
+			UBlueprint* Blueprint = FactoryToUse->CreateBlueprint(AnimSequence, MovieScene, *BlueprintName);
+			check(Blueprint);
+
+			Recording->Guid = MovieScene->AddSpawnable(BlueprintName, Blueprint);
+			
+			// add the anim track
+			UMovieSceneSkeletalAnimationTrack* AnimTrack = MovieScene->AddTrack<UMovieSceneSkeletalAnimationTrack>(Recording->Guid);
+			if(AnimTrack)
+			{
+				AnimTrack->AddNewAnimation(0.0f, AnimSequence);
+			}
+
+			// now add tracks to record
+
+			// we need to create a transform track even if we arent recording transforms
+			TSharedPtr<FMovieScene3DTransformPropertyRecorder> TransformRecorder = MakeShareable(new FMovieScene3DTransformPropertyRecorder);
+			TransformRecorder->CreateSection(Recording->ActorToRecord.Get(), MovieScene, Recording->Guid, Recording->ActorSettings.bRecordTransforms);
+			PropertyRecorders.Add(TransformRecorder);
+
+			// add a visibility track 
+			TSharedPtr<FMovieSceneVisibilityPropertyRecorder> VisibilityRecorder = MakeShareable(new FMovieSceneVisibilityPropertyRecorder);
+			VisibilityRecorder->CreateSection(Recording->ActorToRecord.Get(), MovieScene, Recording->Guid ,Recording->ActorSettings.bRecordVisibility);
+			PropertyRecorders.Add(VisibilityRecorder);
+		}
+
+		if(!bStartedRecordingAllActors)
+		{
+			// if we couldnt start a recording, stop all others
+			TArray<FAssetData> AssetsToCleanUp;
+			if(LevelSequence)
+			{
+				AssetsToCleanUp.Add(LevelSequence);
+			}
+
+			for(UActorRecording* Recording : QueuedRecordings)
+			{
+				Recording->StopRecording();
+				if(Recording->LastRecordedAnimation.IsValid())
+				{
+					AssetsToCleanUp.Add(FAssetData(Recording->LastRecordedAnimation.Get()));
+				}
+			}
+
+			// clean up any assets that we can
+			if(AssetsToCleanUp.Num() > 0)
+			{
+				ObjectTools::DeleteAssets(AssetsToCleanUp, false);
+			}
+		}
+
+		return true;
+	}
+
+	return false;
+} 
+
+bool FSequenceRecorder::StopRecording()
+{
+	CurrentDelay = 0.0f;
+
+	for(UActorRecording* Recording : QueuedRecordings)
+	{
+		Recording->StopRecording();
+	}
+
+	// also stop all dead animation recordings, i.e. ones that use GCed components
+	FAnimationRecorderManager::Get().StopRecordingDeadAnimations();
+
+	const USequenceRecorderSettings* Settings = GetDefault<USequenceRecorderSettings>();
+	if(Settings->bCreateLevelSequence)
+	{
+		if(CurrentSequence.IsValid())
+		{
+			for(auto& PropertyRecorder : PropertyRecorders)
+			{
+				PropertyRecorder->FinalizeSection();
+			}
+
+			ULevelSequence* LevelSequence = CurrentSequence.Get();
+
+			// Stop referencing the sequence so we are listed as 'not recording'
+			CurrentSequence = nullptr;
+
+			UMovieScene* MovieScene = LevelSequence->GetMovieScene();
+			if(MovieScene)
+			{
+				float MinRange = 0.0f;
+				float MaxRange = 0.0f;
+
+				for(UActorRecording* Recording : QueuedRecordings)
+				{
+					if(Recording->LastRecordedAnimation.IsValid())
+					{
+						UAnimSequence* AnimSequence = Recording->LastRecordedAnimation.Get();
+
+						// resize anim sequence section to fit now it is recorded
+						UMovieSceneSkeletalAnimationTrack* AnimTrack = MovieScene->FindTrack<UMovieSceneSkeletalAnimationTrack>(Recording->Guid);
+						check(AnimTrack);
+					
+						UMovieSceneSkeletalAnimationSection* AnimSection = Cast<UMovieSceneSkeletalAnimationSection>(AnimTrack->GetAllSections()[0]);
+						AnimSection->SetEndTime(AnimSequence->GetPlayLength());
+
+						MaxRange = FMath::Max(MaxRange, AnimSequence->GetPlayLength());
+
+						// add a spawn track so we actually get spawned
+						UMovieSceneSpawnTrack* SpawnTrack = MovieScene->AddTrack<UMovieSceneSpawnTrack>(Recording->Guid);
+						if(SpawnTrack)
+						{
+							UMovieSceneBoolSection* SpawnSection = Cast<UMovieSceneBoolSection>(SpawnTrack->CreateNewSection());
+
+							SpawnTrack->AddSection(*SpawnSection);
+							SpawnTrack->SetObjectId(Recording->Guid);
+
+							SpawnSection->SetStartTime(0.0f);
+							SpawnSection->SetEndTime(AnimSequence->GetPlayLength());
+
+							SpawnSection->AddKey(0.0f, true, EMovieSceneKeyInterpolation::Break);
+							SpawnSection->AddKey(AnimSequence->GetPlayLength(), false, EMovieSceneKeyInterpolation::Break);
+						}
+					}
+				}
+
+				MovieScene->SetPlaybackRange(MinRange, MaxRange);
+
+				FAssetRegistryModule::AssetCreated(MovieScene);
+			}
+			return true;
+		}
+	}
+	else
+	{
+		return true;
+	}
+
+	return false;
+}
+
+bool FSequenceRecorder::IsDelaying() const
+{
+	return CurrentDelay > 0.0f;
+}
+
+float FSequenceRecorder::GetCurrentDelay() const
+{
+	return CurrentDelay;
+}
\ No newline at end of file
diff --git a/Engine/Plugins/Editor/SequenceRecorder/Source/SequenceRecorder/Private/SequenceRecorder.h b/Engine/Plugins/Editor/SequenceRecorder/Source/SequenceRecorder/Private/SequenceRecorder.h
new file mode 100644
--- /dev/null
+++ b/Engine/Plugins/Editor/SequenceRecorder/Source/SequenceRecorder/Private/SequenceRecorder.h
@@ -0,0 +1,71 @@
+// Copyright 1998-2015 Epic Games, Inc. All Rights Reserved.
+
+#pragma once
+
+struct FSequenceRecorder
+{
+public:
+	/** Singleton accessor */
+	static FSequenceRecorder& Get();
+
+	/** Starts recording a sequence. */
+	bool StartRecording();
+
+	/** Stops any currently recording sequence */
+	bool StopRecording();
+
+	/** Tick the sequence recorder */
+	void Tick(float DeltaSeconds);
+
+	/** Get whether we are currently delaying a recording */
+	bool IsDelaying() const;
+
+	/** Get the current delay we are waiting for */
+	float GetCurrentDelay() const;
+
+	bool IsRecordingQueued(AActor* Actor) const;
+
+	void StartAllQueuedRecordings();
+
+	void StopAllQueuedRecordings();
+
+	void StopRecordingDeadAnimations();
+
+	void AddNewQueuedRecording(AActor* Actor = nullptr, UAnimSequence* AnimSequence = nullptr, float Length = 0.0f);
+
+	void RemoveQueuedRecording(AActor* Actor);
+
+	void RemoveQueuedRecording(class UActorRecording* Recording);
+
+	bool HasQueuedRecordings() const;
+
+	const TArray<class UActorRecording*>& GetQueuedRecordings() { return QueuedRecordings; }
+
+	bool AreQueuedRecordingsDirty() const { return bQueuedRecordingsDirty; }
+
+	void ResetQueuedRecordingsDirty() { bQueuedRecordingsDirty = false; }
+
+private:
+	/** Starts recording a sequence, possibly delayed */
+	bool StartRecordingInternal();
+
+private:
+	/** Constructor, private - use Get() function */
+	FSequenceRecorder();
+
+	/** Currently recording level sequence, if any */
+	TWeakObjectPtr<class ULevelSequence> CurrentSequence;
+
+	/** The current set of property recorders */
+	TArray<TSharedPtr<class IMovieScenePropertyRecorder>> PropertyRecorders;
+
+	TArray<class UActorRecording*> QueuedRecordings;
+
+	bool bQueuedRecordingsDirty;
+
+	/** The delay we are currently waiting for */
+	float CurrentDelay;
+
+	/** Current recording time */
+	float CurrentTime;
+};
diff --git a/Engine/Plugins/Editor/SequenceRecorder/Source/SequenceRecorder/Private/SequenceRecorderCommands.cpp b/Engine/Plugins/Editor/SequenceRecorder/Source/SequenceRecorder/Private/SequenceRecorderCommands.cpp
new file mode 100644
--- /dev/null
+++ b/Engine/Plugins/Editor/SequenceRecorder/Source/SequenceRecorder/Private/SequenceRecorderCommands.cpp
@@ -0,0 +1,22 @@
+// Copyright 1998-2015 Epic Games, Inc. All Rights Reserved.
+
+#include "SequenceRecorderPrivatePCH.h"
+#include "SequenceRecorderCommands.h"
+#include "SequenceRecorderStyle.h"
+
+#define LOCTEXT_NAMESPACE "SequenceRecorderCommands"
+
+FSequenceRecorderCommands::FSequenceRecorderCommands()
+	: TCommands<FSequenceRecorderCommands>("SequenceRecorder.Common", LOCTEXT("Common", "Common"), NAME_None, FSequenceRecorderStyle::GetStyleSetName())
+{
+}
+
+void FSequenceRecorderCommands::RegisterCommands()
+{
+	UI_COMMAND(RecordAll, "Record", "Record all recordings", EUserInterfaceActionType::Button, FInputChord(EModifierKey::Control | EModifierKey::Shift, EKeys::R));
+	UI_COMMAND(StopAll, "StopAll", "Stop all recordings", EUserInterfaceActionType::Button, FInputChord(EModifierKey::Control | EModifierKey::Shift, EKeys::S));
+	UI_COMMAND(AddRecording, "Add", "Add a new recording", EUserInterfaceActionType::Button, FInputChord());
+	UI_COMMAND(RemoveRecording, "Remove", "Remove selected recording", EUserInterfaceActionType::Button, FInputChord());
+}
+
+#undef LOCTEXT_NAMESPACE
\ No newline at end of file
diff --git a/Engine/Plugins/Editor/SequenceRecorder/Source/SequenceRecorder/Private/SequenceRecorderModule.cpp b/Engine/Plugins/Editor/SequenceRecorder/Source/SequenceRecorder/Private/SequenceRecorderModule.cpp
new file mode 100644
--- /dev/null
+++ b/Engine/Plugins/Editor/SequenceRecorder/Source/SequenceRecorder/Private/SequenceRecorderModule.cpp
@@ -0,0 +1,395 @@
+// Copyright 1998-2016 Epic Games, Inc. All Rights Reserved.
+
+#include "SequenceRecorderPrivatePCH.h"
+#include "ModuleManager.h"
+#include "SSequenceRecorder.h"
+#include "SequenceRecorderStyle.h"
+#include "SDockTab.h"
+#include "SequenceRecorderCommands.h"
+#include "SequenceRecorder.h"
+#include "SequenceRecorderSettings.h"
+
+#define LOCTEXT_NAMESPACE "SequenceRecorder"
+
+static const FName SequenceRecorderTabName(TEXT("SequenceRecorder"));
+
+class FSequenceRecorderModule : public IModuleInterface, private FSelfRegisteringExec
+{
+	/** IModuleInterface implementation */
+	virtual void StartupModule() override
+	{
+#if WITH_EDITOR
+		FSequenceRecorderStyle::Initialize();
+
+		FSequenceRecorderCommands::Register();
+
+		if (GEditor)
+		{
+			// register main tick
+			PostEditorTickHandle = GEditor->OnPostEditorTick().AddStatic(&FSequenceRecorderModule::TickSequenceRecorder);
+		}
+
+		// register Persona recorder
+		FPersonaModule& PersonaModule = FModuleManager::LoadModuleChecked<FPersonaModule>(TEXT("Persona"));
+		PersonaModule.OnIsRecordingActive().BindStatic(&FSequenceRecorderModule::HandlePersonaIsRecordingActive);
+		PersonaModule.OnRecord().BindStatic(&FSequenceRecorderModule::HandlePersonaRecord);
+		PersonaModule.OnStopRecording().BindStatic(&FSequenceRecorderModule::HandlePersonaStopRecording);
+		PersonaModule.OnGetCurrentRecording().BindStatic(&FSequenceRecorderModule::HandlePersonaCurrentRecording);
+		PersonaModule.OnGetCurrentRecordingTime().BindStatic(&FSequenceRecorderModule::HandlePersonaCurrentRecordingTime);
+		PersonaModule.OnTickRecording().BindStatic(&FSequenceRecorderModule::HandlePersonaTickRecording);
+
+		// register 'keep simulation changes' recorder
+		FLevelEditorModule& LevelEditorModule = FModuleManager::LoadModuleChecked<FLevelEditorModule>(TEXT("LevelEditor"));
+		LevelEditorModule.OnCaptureSingleFrameAnimSequence().BindStatic(&FSequenceRecorderModule::HandleCaptureSingleFrameAnimSequence);
+
+		// register level editor menu extender
+		LevelEditorMenuExtenderDelegate = FLevelEditorModule::FLevelViewportMenuExtender_SelectedActors::CreateStatic(&FSequenceRecorderModule::ExtendLevelViewportContextMenu);
+		auto& MenuExtenders = LevelEditorModule.GetAllLevelViewportContextMenuExtenders();
+		MenuExtenders.Add(LevelEditorMenuExtenderDelegate);
+		LevelEditorExtenderDelegateHandle = MenuExtenders.Last().GetHandle();
+
+		// register standalone UI
+		FGlobalTabmanager::Get()->RegisterNomadTabSpawner(SequenceRecorderTabName, FOnSpawnTab::CreateStatic(&FSequenceRecorderModule::SpawnSequenceRecorderTab))
+			.SetGroup(WorkspaceMenu::GetMenuStructure().GetLevelEditorCategory())
+			.SetDisplayName(LOCTEXT("SequenceRecorderTabTitle", "Sequence Recorder"))
+			.SetTooltipText(LOCTEXT("SequenceRecorderTooltipText", "Open the Sequence Recorder tab."))
+			.SetIcon(FSlateIcon(FSequenceRecorderStyle::Get()->GetStyleSetName(), "SequenceRecorder.TabIcon"));
+#endif
+	}
+
+	virtual void ShutdownModule() override
+	{
+#if WITH_EDITOR
+		if (FSlateApplication::IsInitialized())
+		{
+			FGlobalTabmanager::Get()->UnregisterTabSpawner(SequenceRecorderTabName);
+		}
+
+		if(FModuleManager::Get().IsModuleLoaded(TEXT("LevelEditor")))
+		{
+			FLevelEditorModule& LevelEditorModule = FModuleManager::GetModuleChecked<FLevelEditorModule>(TEXT("LevelEditor"));
+			LevelEditorModule.OnCaptureSingleFrameAnimSequence().Unbind();
+			LevelEditorModule.GetAllLevelViewportContextMenuExtenders().RemoveAll([&](const FLevelEditorModule::FLevelViewportMenuExtender_SelectedActors& Delegate) {
+				return Delegate.GetHandle() == LevelEditorExtenderDelegateHandle;
+			});
+		}
+
+		if (FModuleManager::Get().IsModuleLoaded(TEXT("Persona")))
+		{
+			FPersonaModule& PersonaModule = FModuleManager::GetModuleChecked<FPersonaModule>(TEXT("Persona"));
+			PersonaModule.OnIsRecordingActive().Unbind();
+			PersonaModule.OnRecord().Unbind();
+			PersonaModule.OnStopRecording().Unbind();
+			PersonaModule.OnGetCurrentRecording().Unbind();
+			PersonaModule.OnGetCurrentRecordingTime().Unbind();
+			PersonaModule.OnTickRecording().Unbind();
+		}
+
+		if (GEditor)
+		{
+			GEditor->OnPostEditorTick().Remove(PostEditorTickHandle);
+		}
+
+		FSequenceRecorderStyle::Shutdown();
+#endif
+	}
+
+	// FSelfRegisteringExec implementation
+	virtual bool Exec(UWorld* InWorld, const TCHAR* Cmd, FOutputDevice& Ar) override
+	{
+#if WITH_EDITOR
+		if (FParse::Command(&Cmd, TEXT("RecordAnimation")))
+		{
+			return HandleRecordAnimationCommand(InWorld, Cmd, Ar);
+		}
+		else if (FParse::Command(&Cmd, TEXT("StopRecordingAnimation")))
+		{
+			return HandleStopRecordAnimationCommand(InWorld, Cmd, Ar);
+		}
+#endif
+		return false;
+	}
+
+	static bool HandleRecordAnimationCommand(UWorld* InWorld, const TCHAR* InStr, FOutputDevice& Ar)
+	{
+#if WITH_EDITOR
+		const TCHAR* Str = InStr;
+		// parse actor name
+		TCHAR ActorName[128];
+		AActor* FoundActor = nullptr;
+		if (FParse::Token(Str, ActorName, ARRAY_COUNT(ActorName), 0))
+		{
+			FString const ActorNameStr = FString(ActorName);
+			for (ULevel const* Level : InWorld->GetLevels())
+			{
+				if (Level)
+				{
+					for (AActor* Actor : Level->Actors)
+					{
+						if (Actor)
+						{
+							if (Actor->GetName() == ActorNameStr)
+							{
+								FoundActor = Actor;
+								break;
+							}
+						}
+					}
+				}
+			}
+		}
+
+		if (FoundActor)
+		{
+			USkeletalMeshComponent* const SkelComp = FoundActor->FindComponentByClass<USkeletalMeshComponent>();
+			if (SkelComp)
+			{
+				TCHAR AssetPath[256];
+				FParse::Token(Str, AssetPath, ARRAY_COUNT(AssetPath), 0);
+				FString const AssetName = FPackageName::GetLongPackageAssetName(AssetPath);
+				return FAnimationRecorderManager::Get().RecordAnimation(SkelComp, AssetPath, AssetName);
+			}
+		}
+#endif
+		return false;
+	}
+
+	static bool HandleStopRecordAnimationCommand(UWorld* InWorld, const TCHAR* InStr, FOutputDevice& Ar)
+	{
+#if WITH_EDITOR
+		const TCHAR* Str = InStr;
+
+		// parse actor name
+		TCHAR ActorName[128];
+		AActor* FoundActor = nullptr;
+		bool bStopAll = false;
+		if (FParse::Token(Str, ActorName, ARRAY_COUNT(ActorName), 0))
+		{
+			FString const ActorNameStr = FString(ActorName);
+
+			if (ActorNameStr.ToLower() == TEXT("all"))
+			{
+				bStopAll = true;
+			}
+			else if (InWorld)
+			{
+				// search for the actor by name
+				for (ULevel* Level : InWorld->GetLevels())
+				{
+					if (Level)
+					{
+						for (AActor* Actor : Level->Actors)
+						{
+							if (Actor)
+							{
+								if (Actor->GetName() == ActorNameStr)
+								{
+									FoundActor = Actor;
+									break;
+								}
+							}
+						}
+					}
+				}
+			}
+		}
+
+		if (bStopAll)
+		{
+			FAnimationRecorderManager::Get().StopRecordingAllAnimations();
+			return true;
+		}
+		else if (FoundActor)
+		{
+			USkeletalMeshComponent* const SkelComp = FoundActor->FindComponentByClass<USkeletalMeshComponent>();
+			if (SkelComp)
+			{
+				FAnimationRecorderManager::Get().StopRecordingAnimation(SkelComp);
+				return true;
+			}
+		}
+
+#endif
+		return false;
+	}
+
+	static void TickSequenceRecorder(float DeltaSeconds)
+	{
+		if (!IsRunningDedicatedServer() && !IsRunningCommandlet())
+		{
+			FSequenceRecorder::Get().Tick(DeltaSeconds);
+		}
+	}
+
+	static UAnimSequence* HandleCaptureSingleFrameAnimSequence(USkeletalMeshComponent* Component)
+	{
+		FAnimationRecorder Recorder;
+		if (Recorder.TriggerRecordAnimation(Component))
+		{
+			class UAnimSequence * Sequence = Recorder.GetAnimationObject();
+			if (Sequence)
+			{
+				Recorder.StopRecord(false);
+				return Sequence;
+			}
+		}
+
+		return nullptr;
+	}
+
+	static void HandlePersonaIsRecordingActive(USkeletalMeshComponent* Component, bool& bIsRecording)
+	{
+		bIsRecording = FAnimationRecorderManager::Get().IsRecording(Component);
+	}
+
+	static void HandlePersonaRecord(USkeletalMeshComponent* Component)
+	{
+		FAnimationRecorderManager::Get().RecordAnimation(Component);
+	}
+
+	static void HandlePersonaStopRecording(USkeletalMeshComponent* Component)
+	{
+		FAnimationRecorderManager::Get().StopRecordingAnimation(Component);
+	}
+
+	static void HandlePersonaTickRecording(USkeletalMeshComponent* Component, float DeltaSeconds)
+	{
+		FAnimationRecorderManager::Get().Tick(Component, DeltaSeconds);
+	}
+
+	static void HandlePersonaCurrentRecording(USkeletalMeshComponent* Component, UAnimSequence*& OutSequence)
+	{
+		OutSequence = FAnimationRecorderManager::Get().GetCurrentlyRecordingSequence(Component);
+	}
+
+	static void HandlePersonaCurrentRecordingTime(USkeletalMeshComponent* Component, float& OutTime)
+	{
+		OutTime = FAnimationRecorderManager::Get().GetCurrentRecordingTime(Component);
+	}
+
+	static TSharedRef<SDockTab> SpawnSequenceRecorderTab(const FSpawnTabArgs& SpawnTabArgs)
+	{
+		const TSharedRef<SDockTab> MajorTab =
+			SNew(SDockTab)
+			.Icon(FSequenceRecorderStyle::Get()->GetBrush("SequenceRecorder.TabIcon"))
+			.TabRole(ETabRole::NomadTab);
+
+		MajorTab->SetContent(SNew(SSequenceRecorder));
+
+		return MajorTab;
+	}
+
+	static TSharedRef<FExtender> ExtendLevelViewportContextMenu(const TSharedRef<FUICommandList> CommandList, const TArray<AActor*> SelectedActors)
+	{
+		TSharedRef<FExtender> Extender(new FExtender());
+
+		// check if actors are already registered for recording
+		bool bCanAddRecording = false;
+		bool bCanRemoveRecording = false;
+
+		FSequenceRecorder& SequenceRecorder = FSequenceRecorder::Get();
+
+		for (AActor* Actor : SelectedActors)
+		{
+			const bool bIsQueued = SequenceRecorder.IsRecordingQueued(Actor);
+			bCanRemoveRecording |= bIsQueued;
+			bCanAddRecording |= !bIsQueued;
+		}
+
+		if (bCanAddRecording || bCanRemoveRecording)
+		{
+			// Add the sequence recorder sub-menu extender
+			Extender->AddMenuExtension(
+				"ActorSelectVisibilityLevels",
+				EExtensionHook::After,
+				nullptr,
+				FMenuExtensionDelegate::CreateStatic(&FSequenceRecorderModule::CreateLevelViewportContextMenuEntries, bCanAddRecording, bCanRemoveRecording));
+		}
+
+		return Extender;
+	}
+
+	static void CreateLevelViewportContextMenuEntries(FMenuBuilder& MenuBuilder, bool bCanAddRecording, bool bCanRemoveRecording)
+	{
+		MenuBuilder.BeginSection("SequenceRecording", LOCTEXT("SequenceRecordingLevelEditorHeading", "Sequence Recording"));
+
+		if (bCanAddRecording)
+		{
+			FUIAction Action_AddQuickRecording(FExecuteAction::CreateStatic(&FSequenceRecorderModule::AddActorsForQuickRecording));
+
+			MenuBuilder.AddMenuEntry(
+				LOCTEXT("MenuExtensionAddQuickRecording", "Trigger Actor Recording"),
+				LOCTEXT("MenuExtensionAddQuickRecording_Tooltip", "Set up the selected actors for recording and trigger recording immediately"),
+				FSlateIcon(),
+				Action_AddQuickRecording,
+				NAME_None,
+				EUserInterfaceActionType::Button);
+
+			FUIAction Action_AddRecording(FExecuteAction::CreateStatic(&FSequenceRecorderModule::AddActorsForRecording));
+
+			MenuBuilder.AddMenuEntry(
+				LOCTEXT("MenuExtensionAddRecording", "Queue Actor Recording"),
+				LOCTEXT("MenuExtensionAddRecording_Tooltip", "Queue up the selected actors for recording"),
+				FSlateIcon(),
+				Action_AddRecording,
+				NAME_None,
+				EUserInterfaceActionType::Button);
+		}
+
+		if (bCanRemoveRecording)
+		{
+			FUIAction Action_RemoveRecording(FExecuteAction::CreateStatic(&FSequenceRecorderModule::RemoveActorsForRecording));
+
+			MenuBuilder.AddMenuEntry(
+				LOCTEXT("MenuExtensionRemoveRecording", "Removed Queued Recording"),
+				LOCTEXT("MenuExtensionRemoveRecording_Tooltip", "Remove the selected actors from the recording queue"),
+				FSlateIcon(),
+				Action_RemoveRecording,
+				NAME_None,
+				EUserInterfaceActionType::Button);
+		}
+
+		MenuBuilder.EndSection();
+	}
+
+	static void AddActorsForRecording()
+	{
+		FSequenceRecorder& SequenceRecorder = FSequenceRecorder::Get();
+
+		TArray<AActor*> SelectedActors;
+		GEditor->GetSelectedActors()->GetSelectedObjects<AActor>(SelectedActors);
+		for(AActor* Actor : SelectedActors)
+		{
+			SequenceRecorder.AddNewQueuedRecording(Actor, nullptr, GetDefault<USequenceRecorderSettings>()->SequenceLength);
+		}
+	}
+
+	static void AddActorsForQuickRecording()
+	{
+		AddActorsForRecording();
+
+		FSequenceRecorder::Get().StartRecording();
+	}
+
+	static void RemoveActorsForRecording()
+	{
+		FSequenceRecorder& SequenceRecorder = FSequenceRecorder::Get();
+
+		TArray<AActor*> SelectedActors;
+		GEditor->GetSelectedActors()->GetSelectedObjects<AActor>(SelectedActors);
+		for (AActor* Actor : SelectedActors)
+		{
+			SequenceRecorder.RemoveQueuedRecording(Actor);
+		}
+	}
+
+	FDelegateHandle PostEditorTickHandle;
+
+	FLevelEditorModule::FLevelViewportMenuExtender_SelectedActors LevelEditorMenuExtenderDelegate;
+
+	FDelegateHandle LevelEditorExtenderDelegateHandle;
+};
+
+IMPLEMENT_MODULE( FSequenceRecorderModule, SequenceRecorder )
+
+#undef LOCTEXT_NAMESPACE
\ No newline at end of file
diff --git a/Engine/Plugins/Editor/SequenceRecorder/Source/SequenceRecorder/Private/SequenceRecorderPrivatePCH.h b/Engine/Plugins/Editor/SequenceRecorder/Source/SequenceRecorder/Private/SequenceRecorderPrivatePCH.h
new file mode 100644
--- /dev/null
+++ b/Engine/Plugins/Editor/SequenceRecorder/Source/SequenceRecorder/Private/SequenceRecorderPrivatePCH.h
@@ -0,0 +1,14 @@
+// Copyright 1998-2016 Epic Games, Inc. All Rights Reserved.
+
+#include "Editor.h"
+#include "LevelEditor.h"
+#include "PersonaModule.h"
+#include "AnimationRecordingSettings.h"
+#include "AnimationRecorder.h"
+#include "SlateBasics.h"
+#include "WorkspaceMenuStructureModule.h"
+#include "PropertyEditing.h"
+#include "MovieScene.h"
+#include "KeyParams.h"
+#include "IKeyframeSection.h"
+#include "MovieSceneSection.h"
\ No newline at end of file
diff --git a/Engine/Plugins/Editor/SequenceRecorder/Source/SequenceRecorder/Private/SequenceRecorderSettings.cpp b/Engine/Plugins/Editor/SequenceRecorder/Source/SequenceRecorder/Private/SequenceRecorderSettings.cpp
new file mode 100644
--- /dev/null
+++ b/Engine/Plugins/Editor/SequenceRecorder/Source/SequenceRecorder/Private/SequenceRecorderSettings.cpp
@@ -0,0 +1,15 @@
+// Copyright 1998-2016 Epic Games, Inc. All Rights Reserved.
+
+#include "SequenceRecorderPrivatePCH.h"
+#include "SequenceRecorderSettings.h"
+
+USequenceRecorderSettings::USequenceRecorderSettings(const FObjectInitializer& ObjectInitializer)
+	: Super(ObjectInitializer)
+{
+	bCreateLevelSequence = true;
+	SequenceLength = FAnimationRecordingSettings::DefaultMaximumLength;
+	RecordingDelay = 4.0f;
+	SequenceName = TEXT("RecordedSequence");
+	AnimationSubDirectory = TEXT("Animations");
+	SequenceRecordingBasePath.Path = TEXT("Cinematics/Sequences");
+}
diff --git a/Engine/Plugins/Editor/SequenceRecorder/Source/SequenceRecorder/Private/SequenceRecorderSettings.h b/Engine/Plugins/Editor/SequenceRecorder/Source/SequenceRecorder/Private/SequenceRecorderSettings.h
new file mode 100644
--- /dev/null
+++ b/Engine/Plugins/Editor/SequenceRecorder/Source/SequenceRecorder/Private/SequenceRecorderSettings.h
@@ -0,0 +1,40 @@
+// Copyright 1998-2016 Epic Games, Inc. All Rights Reserved.
+
+#pragma once
+
+#include "SequenceRecorderSettings.generated.h"
+
+UCLASS(config=Editor, MinimalAPI)
+class USequenceRecorderSettings : public UObject
+{
+	GENERATED_UCLASS_BODY()
+
+public:
+	/** Whether to create a level sequence when recording. Actors and animations will be inserted into this sequence */
+	UPROPERTY(EditAnywhere, Category = "Sequence Recording")
+	bool bCreateLevelSequence;
+
+	/** The length of the recorded sequence */
+	UPROPERTY(EditAnywhere, Category = "Sequence Recording")
+	float SequenceLength;
+
+	/** Delay that we will use before starting recording */
+	UPROPERTY(EditAnywhere, Category = "Sequence Recording")
+	float RecordingDelay;
+
+	/** The base name of the sequence to record to. This name will also be used to auto-generate any assets created by this recording. */
+	UPROPERTY(EditAnywhere, Category = "Sequence Recording")
+	FString SequenceName;
+
+	/** Base path for this recording. Sub-assets will be created in subdirectories as specified */
+	UPROPERTY(EditAnywhere, Category = "Sequence Recording", meta=(RelativeToGameContentDir))
+	FDirectoryPath SequenceRecordingBasePath;
+
+	/** The name of the subdirectory animations will be placed in. Leave this empty to place into the same directory as the sequence base path */
+	UPROPERTY(EditAnywhere, AdvancedDisplay, Category = "Sequence Recording")
+	FString AnimationSubDirectory;
+
+	/** When set, record a transform track and record animations in local space */
+	UPROPERTY(EditAnywhere, Category = "Sequence Recording")
+	bool bRecordTransformsAsSeperateTrack;
+};
\ No newline at end of file
