Copying //UE4/Dev-Sequencer to //UE4/Main

==========================
MAJOR FEATURES + CHANGES
==========================

Change 2859626 on 2016/02/08 by [REDACTED_PERSON]

	Editor: Added SaveAs functionality to content asset editors

Change 2859666 on 2016/02/08 by [REDACTED_PERSON]

	Sequencer: Fix crash in CheckForWorldGCLeaks when loading a new map because spawnables are left behind.

Change 2859685 on 2016/02/08 by [REDACTED_PERSON]

	Sequencer: Add prompt to save sub level sequences if they are dirty

Change 2859715 on 2016/02/08 by [REDACTED_PERSON]

	Adding actor spawning recording

	Actors are queued for record on spawn then added to the list like manually-specifed ones.
	Changed almost everything about UActorRecording. We now record on a per-component basis, with property tracks encapsulated in each actor recording. Much effort is expended to make sure that the correct components are owned by their respective actors, as we can add and remove components at runtime, but they must be created up-front in the UMovieScene Blueprints. We go as far as to add our own SCS nodes to make sure components are correctly spawned.
	Fixed infinite loop in FSequencer::AddSpawnable.
	Fixed visibility track instance to work with scene components as well as actors.
	Fixed particle track instance to work on UParticleSystemComponent rather than just AEmitters.
	Added particle recorder.
	Moved animation recording into an animation property recorder rather than having it as a special case. This still uses the animation recorder under the hood.
	Moved old-style Matinee animation control into FMovieSceneSkeletalAnimationTrackInstance & made this work on USkeletalMeshComponents directly, rather than via the old interface.
	Exposed SetMatineeAnimPositionInner and PreviewMatineeSetAnimPositionInner in FAnimMontageInstance so those utility functions can be used externally to Engine.
	Added a predicate version of UMovieScene::FindPossessable.
	Exposed UMovieSceneParticleSection::AddKey externally via MOVIESCENETRACKS_API so I can programmatically add keys.
	Fixed a crash in FScalableFloatDetails::CustomizeHeader when selecting PIE projectiles in Orion.
	Moved all recorders over to recording Actors or Components & store UObjects instead of AActors.
	Allowed skeletal animation tracks on components as well as actors.

Change 2862675 on 2016/02/10 by [REDACTED_PERSON]

	Sequencer: Add option to link the sequencer curve editor with the sequencer timeline.

	Under General Options->Link Curve Editor Time Range. The default is false, so the sequencer and curve editor have separate time ranges.

Change 2862699 on 2016/02/11 by [REDACTED_PERSON]

	Sequencer: Added a playback status of jumping which the AudioTrack and Skeletal Mesh Track (anim notifies) ignores for updates. This is used to updating thumbnail at certain times.

Change 2862712 on 2016/02/11 by [REDACTED_PERSON]

	Sequencer: Fix spawnables firing off their particles. Disable auto activate on spawnable components

Change 2862719 on 2016/02/11 by [REDACTED_PERSON]

	Editor: Refactored detail customizations for colors, rotators, vectors

	- broke color and rotator customizations out into their own files
	- added vector customizations (placeholder)
	- cleaned up localization namespaces, forward declarations

Change 2866454 on 2016/02/14 by [REDACTED_PERSON]

	Sequencer: Removed ULevelEditorSequencerSettings; moved default settings into INI

Change 2866455 on 2016/02/14 by [REDACTED_PERSON]

	Sequence recorder can now record replays

	Added extra edtior-only UI to the replay playback controls to record sequences. Curretnly very placeholder: only records the entire sequence and provides no feedback in the UI if it is recording.
	Fixed bindings to recorded objects not working in various circumstances. Added the ability to manually create a binding.
	Recompiled actor blueprints post-record if we added components.
	Fixed a null ptr dereference in FOrionTeamUIInfo::Update.
	Removed tolerances when reducing tracks - they are now 'very small'.
	Added actor filter so actors of certain classes can be recorded.

Change 2866458 on 2016/02/14 by [REDACTED_PERSON]

	Sequencer: Fix anim notifies that fire at shot cuts. Anim notifies are fired from the last position to the current position. When jumping cuts, we want the delta to be 0 so that the anim notifies before the shot are not fired off.

Change 2866459 on 2016/02/14 by [REDACTED_PERSON]

	Sequencer: Add option to toggle visibility of combined keys

Change 2866466 on 2016/02/14 by [REDACTED_PERSON]

	Sequencer - Add a track for controlling streamed level visibilty and remove visibility code from the master level blueprint.

Change 2866470 on 2016/02/14 by [REDACTED_PERSON]

	Sequencer: Add return value to indicate data has changed when a section has been added. This fixes a bug where creating a new section doesn't seem to add a key.

Change 2866481 on 2016/02/14 by [REDACTED_PERSON]

	Sequencer: Implemented Presets for adding tracks automatically based on actor type (UE-24513)

Change 2866482 on 2016/02/14 by [REDACTED_PERSON]

	Sequencer: Allow for any actor that has a camera component to be a camera cut.

Change 2866484 on 2016/02/14 by [REDACTED_PERSON]

	Added in/out times to sequence recording

	Also added the optional ability to record different actor types (heroes, projectiles, minions).

Change 2866495 on 2016/02/14 by [REDACTED_PERSON]

	Sequencer: Need to limit camera control to the section bounds of the camera cut otherwise, control won't be relinquished back to player at the end of the playback.

diff --git a/Engine/Config/BaseEditorPerProjectUserSettings.ini b/Engine/Config/BaseEditorPerProjectUserSettings.ini
--- a/Engine/Config/BaseEditorPerProjectUserSettings.ini
+++ b/Engine/Config/BaseEditorPerProjectUserSettings.ini
@@ -714,4 +714,16 @@ DefaultBrushRadius=128
 +Node=(Class=MaterialExpressionConstant4Vector Key=Four Shift=false Ctrl=false Alt=false)
 
 [WidgetTemplatesExpanded]
-Common=True
\ No newline at end of file
+Common=True
+
+[LevelSequenceEditor SequencerSettings]
+bKeyInterpPropertiesOnly=true
+bShowRangeSlider=true
+bKeepPlayRangeInSectionBounds=false
+TimeSnapInterval=0.033334
+ZeroPadFrames=4
+
+[/Script/LevelSequenceEditor.LevelSequenceEditorSettings]
++TrackSettings=(MatchingActorClass=/Script/Engine.CameraActor,DefaultTracks=(/Script/MovieSceneTracks.MovieScene3DTransformTrack))
++TrackSettings=(MatchingActorClass=/Script/Engine.CineCameraActor,DefaultPropertyTracks=((ComponentPath="CameraComponent",PropertyPath="CurrentFocalLength")))
++TrackSettings=(MatchingActorClass=/Script/Engine.Light,DefaultTracks=(None),DefaultPropertyTracks=((ComponentPath="LightComponent0",PropertyPath="Intensity"),(ComponentPath="LightComponent0",PropertyPath="LightColor")))
diff --git a/Engine/Config/BaseEngine.ini b/Engine/Config/BaseEngine.ini
--- a/Engine/Config/BaseEngine.ini
+++ b/Engine/Config/BaseEngine.ini
@@ -613,6 +613,10 @@ bUseStreamingPause=false
 +TaggedPropertyRedirects=(ClassName="MovieSceneShotSection",OldPropertyName="ShotNumber",NewPropertyName="CameraCutNumber")
 +TaggedPropertyRedirects=(ClassName="MovieScene",OldPropertyName="ShotTrack",NewPropertyName="CameraCutTrack")
 
+; CineCameraActor moved to CinematicCamera
++ActiveClassRedirects=(OldClassName="CineCameraActor",NewClassName="/Script/CinematicCamera.CineCameraActor")
++ActiveClassRedirects=(OldClassName="CineCameraComponent",NewClassName="/Script/CinematicCamera.CineCameraComponent")
+
 ; We renamed EditorGameAgnosticSettings to be less imposing
 +ActiveClassRedirects=(OldClassName="EditorGameAgnosticSettings",NewClassName="/Script/UnrealEd.EditorSettings")
 +ActiveClassRedirects=(OldClassName="EditorUserSettings",NewClassName="/Script/UnrealEd.EditorPerProjectUserSettings")
@@ -2559,4 +2563,24 @@ bBuildForD3D12=false
 +Settings=(ScreenSize=.3,OptimizationSettings=(NumOfTrianglesPercentage=.25))
 +Settings=(ScreenSize=.1,OptimizationSettings=(NumOfTrianglesPercentage=.125))
 +Settings=(ScreenSize=.06,OptimizationSettings=(NumOfTrianglesPercentage=.06))
+[/Script/CinematicCamera.CineCameraComponent]
++FilmbackPresets=(Name="Super 8mm",FilmbackSettings=(SensorWidth=5.79,SensorHeight=4.01))
++FilmbackPresets=(Name="Super 16mm",FilmbackSettings=(SensorWidth=12.52,SensorHeight=7.41))
++FilmbackPresets=(Name="Super 35mm 3 Perf",FilmbackSettings=(SensorWidth=24.9,SensorHeight=12))
++FilmbackPresets=(Name="Super 35mm 4 Perf",FilmbackSettings=(SensorWidth=24.89,SensorHeight=18.67))
++FilmbackPresets=(Name="Four Thirds",FilmbackSettings=(SensorWidth=17.3,SensorHeight=13))
++FilmbackPresets=(Name="35mm Full Frame",FilmbackSettings=(SensorWidth=36,SensorHeight=24))
++FilmbackPresets=(Name="APS-C",FilmbackSettings=(SensorWidth=23.6,SensorHeight=15.7))
++FilmbackPresets=(Name="APS-C (Canon)",FilmbackSettings=(SensorWidth=22.2,SensorHeight=14.8))
++FilmbackPresets=(Name="APS-H",FilmbackSettings=(SensorWidth=27.9,SensorHeight=18.6))
++FilmbackPresets=(Name="65mm",FilmbackSettings=(SensorWidth=52.48,SensorHeight=23.01))
++FilmbackPresets=(Name="IMAX 70mm",FilmbackSettings=(SensorWidth=70.41,SensorHeight=56.63))
++LensPresets=(Name="12mm Prime f/2.8",LensSettings=(MinFocalLength=12,MaxFocalLength=12,MinFStop=2.8,MaxFStop=22)
++LensPresets=(Name="30mm Prime f/1.4",LensSettings=(MinFocalLength=30,MaxFocalLength=30,MinFStop=1.4,MaxFStop=22)
++LensPresets=(Name="50mm Prime f/1.8",LensSettings=(MinFocalLength=50,MaxFocalLength=50,MinFStop=1.8,MaxFStop=22)
++LensPresets=(Name="85mm Prime f/1.8",LensSettings=(MinFocalLength=85,MaxFocalLength=85,MinFStop=1.8,MaxFStop=22)
++LensPresets=(Name="105mm Prime f/2",LensSettings=(MinFocalLength=105,MaxFocalLength=105,MinFStop=2,MaxFStop=22,)
++LensPresets=(Name="200mm Prime f/2",LensSettings=(MinFocalLength=200,MaxFocalLength=200,MinFStop=2,MaxFStop=22)
++LensPresets=(Name="24-70mm Zoom f/2.8",LensSettings=(MinFocalLength=24,MaxFocalLength=70,MinFStop=2.8,MaxFStop=22))
++LensPresets=(Name="70-200mm Zoom f/2.8",LensSettings=(MinFocalLength=70,MaxFocalLength=200,MinFStop=2.8,MaxFStop=22))
 
diff --git a/Engine/Plugins/Editor/SequenceRecorder/Source/SequenceRecorder/Private/ActorRecording.cpp b/Engine/Plugins/Editor/SequenceRecorder/Source/SequenceRecorder/Private/ActorRecording.cpp
--- a/Engine/Plugins/Editor/SequenceRecorder/Source/SequenceRecorder/Private/ActorRecording.cpp
+++ b/Engine/Plugins/Editor/SequenceRecorder/Source/SequenceRecorder/Private/ActorRecording.cpp
@@ -4,113 +4,512 @@
 #include "ActorRecording.h"
 #include "SequenceRecorderSettings.h"
 #include "SequenceRecorderUtils.h"
+#include "IMovieScenePropertyRecorder.h"
+#include "MovieSceneSpawnPropertyRecorder.h"
+#include "MovieScene3DTransformPropertyRecorder.h"
+#include "MovieSceneVisibilityPropertyRecorder.h"
+#include "MovieSceneAnimationPropertyRecorder.h"
+#include "MovieSceneParticleTrackPropertyRecorder.h"
+#include "MovieScene3DAttachPropertyRecorder.h"
+#include "AssetSelection.h"
+#include "MovieSceneSkeletalAnimationTrack.h"
+#include "MovieSceneSkeletalAnimationSection.h"
+#include "MovieSceneSpawnTrack.h"
+#include "MovieScene3DTransformTrack.h"
+#include "MovieScene3DTransformSection.h"
+#include "MovieSceneBoolSection.h"
+#include "KismetEditorUtilities.h"
+#include "BlueprintEditorUtils.h"
 
-bool UActorRecording::StartRecording()
+UActorRecording::UActorRecording(const FObjectInitializer& ObjectInitializer)
+	: Super(ObjectInitializer)
 {
-	LastRecordedAnimation = nullptr;
+	bWasSpawnedPostRecord = false;
+	MeshComponentUpdateFlag = EMeshComponentUpdateFlag::AlwaysTickPoseAndRefreshBones;
+	bEnableUpdateRateOptimizations = false;
+	Guid.Invalidate();
+	bNewComponentAddedWhileRecording = false;
+}
+
+bool UActorRecording::IsRelevantForRecording(AActor* Actor)
+{
+	TInlineComponentArray<USceneComponent*> SceneComponents(Actor);
+
+	for(USceneComponent* SceneComponent : SceneComponents)
+	{
+		if(SceneComponent->IsA<USkeletalMeshComponent>() ||
+			SceneComponent->IsA<UStaticMeshComponent>() ||
+			SceneComponent->IsA<UParticleSystemComponent>())
+		{
+			return true;
+		}
+	}
+
+	return false;
+}
+
+bool UActorRecording::StartRecording(ULevelSequence* CurrentSequence, float CurrentSequenceTime)
+{
+	bNewComponentAddedWhileRecording = false;
 
 	if(ActorToRecord.IsValid())
 	{
-		// @todo: animation recording: allow recording of multiple skeletal meshes per actor
-		USkeletalMeshComponent* SkeletalMeshComponent = ActorToRecord.Get()->FindComponentByClass<USkeletalMeshComponent>();
-		if(SkeletalMeshComponent)
+		if(CurrentSequence != nullptr)
 		{
-			// turn off URO and make sure we always update even if out of view
-			bEnableUpdateRateOptimizations = SkeletalMeshComponent->bEnableUpdateRateOptimizations;
-			MeshComponentUpdateFlag = SkeletalMeshComponent->MeshComponentUpdateFlag;
+			StartRecordingActorProperties(CurrentSequence, CurrentSequenceTime);
+		}
+		else
+		{
+			TSharedPtr<FMovieSceneAnimationPropertyRecorder> AnimationRecorder = MakeShareable(new FMovieSceneAnimationPropertyRecorder(AnimationSettings, TargetAnimation.Get()));
+			AnimationRecorder->CreateSection(ActorToRecord.Get(), nullptr, FGuid(), 0.0f, true);
+			PropertyRecorders.Add(AnimationRecorder);			
+		}
+	}
+
+	return true;
+}
+
+static FString GetUniqueSpawnableName(UMovieScene* MovieScene, const FString& BaseName)
+{
+	FString BlueprintName = BaseName;
+	auto DuplName = [&](FMovieSceneSpawnable& InSpawnable)
+	{
+		return InSpawnable.GetName() == BlueprintName;
+	};
+
+	int32 Index = 2;
+	FString UniqueString;
+	while (MovieScene->FindSpawnable(DuplName))
+	{
+		BlueprintName.RemoveFromEnd(UniqueString);
+		UniqueString = FString::Printf(TEXT(" (%d)"), Index++);
+		BlueprintName += UniqueString;
+	}
+
+	return BlueprintName;
+}
 
-			SkeletalMeshComponent->bEnableUpdateRateOptimizations = false;
-			SkeletalMeshComponent->MeshComponentUpdateFlag = EMeshComponentUpdateFlag::AlwaysTickPoseAndRefreshBones;
+void UActorRecording::GetSceneComponents(TArray<USceneComponent*>& OutArray, bool bIncludeNonCDO/*=true*/)
+{
+	// it is not enough to just go through the owned components array here
+	// we need to traverse the scene component hierarchy as well, as some components may be 
+	// owned by other actors (e.g. for pooling) and some may not be part of the hierarchy
+	if(ActorToRecord.IsValid())
+	{
+		USceneComponent* RootComponent = ActorToRecord->GetRootComponent();
+		if(RootComponent)
+		{
+			// note: GetChildrenComponents clears array!
+			RootComponent->GetChildrenComponents(true, OutArray);
+			OutArray.Add(RootComponent);
+		}
 
-			UAnimSequence* Sequence = nullptr;
-			const bool bValidSpecifiedAnimation = (bSpecifyTargetAnimation && TargetAnimation.IsValid());
-			if(bValidSpecifiedAnimation)
+		// add owned components that are *not* part of the hierarchy
+		TInlineComponentArray<USceneComponent*> OwnedComponents(ActorToRecord.Get());
+		for(USceneComponent* OwnedComponent : OwnedComponents)
+		{
+			if(OwnedComponent->AttachParent == nullptr && OwnedComponent != RootComponent)
 			{
-				Sequence = TargetAnimation.Get();
+				OutArray.Add(OwnedComponent);
 			}
+		}
 
-			bool bValidAutoGeneratedAnimation = false;
-			if(!bSpecifyTargetAnimation)
-			{
-				// build an asset path
-				const USequenceRecorderSettings* Settings = GetDefault<USequenceRecorderSettings>();
+		if(!bIncludeNonCDO)
+		{
+			AActor* CDO = Cast<AActor>(ActorToRecord->GetClass()->GetDefaultObject());
+
+			auto ShouldRemovePredicate = [&](UActorComponent* PossiblyRemovedComponent)
+				{
+					// try to find a component with this name in the CDO
+					UActorComponent* const* FoundComponent = CDO->GetComponents().FindByPredicate([&](UActorComponent* SearchComponent) 
+					{ 
+						return SearchComponent->GetClass() == PossiblyRemovedComponent->GetClass() &&
+							   SearchComponent->GetFName() == PossiblyRemovedComponent->GetFName(); 
+					} );
+
+					// remove if its not found
+					return FoundComponent == nullptr;
+				};
+
+			OutArray.RemoveAllSwap(ShouldRemovePredicate);
+		}
+	}
+}
 
-				FString AssetPath(TEXT("/Game"));
-				AssetPath /= Settings->SequenceRecordingBasePath.Path;
-				if(Settings->AnimationSubDirectory.Len() > 0)
+void UActorRecording::SyncTrackedComponents(bool bIncludeNonCDO/*=true*/)
+{
+	TArray<USceneComponent*> NewComponentArray;
+	GetSceneComponents(NewComponentArray, bIncludeNonCDO);
+
+	TrackedComponents.Reset(NewComponentArray.Num());
+	for(USceneComponent* SceneComponent : NewComponentArray)
+	{
+		TrackedComponents.Add(SceneComponent);
+	}
+}
+
+bool UActorRecording::ValidComponent(USceneComponent* SceneComponent) const
+{
+	if(SceneComponent != nullptr)
+	{
+		const bool bIsValidComponent = 
+			SceneComponent->IsA<USkeletalMeshComponent>() ||
+			SceneComponent->IsA<UStaticMeshComponent>() ||
+			SceneComponent->IsA<UParticleSystemComponent>();
+
+		return bIsValidComponent;
+	}
+
+	return false;
+}
+
+void UActorRecording::StartRecordingActorProperties(ULevelSequence* CurrentSequence, float CurrentSequenceTime)
+{
+	if(CurrentSequence != nullptr)
+	{
+		// set up our spawnable for this actor
+		UMovieScene* MovieScene = CurrentSequence->GetMovieScene();
+
+		AActor* Actor = ActorToRecord.Get();
+		FString BlueprintName = GetUniqueSpawnableName(MovieScene, Actor->GetName());
+
+		UClass* ActorClass = Actor->GetClass();
+			
+		UBlueprint* Blueprint = nullptr;
+
+		if(ActorClass->IsChildOf<AWorldSettings>())
+		{
+			// we need to create our own substitute blueprint for world settings actors as they are not blueprintable by default
+			Blueprint = FKismetEditorUtilities::CreateBlueprint(AActor::StaticClass(), MovieScene, *BlueprintName, EBlueprintType::BPTYPE_Normal, UBlueprint::StaticClass(), UBlueprintGeneratedClass::StaticClass(), NAME_Name);
+		}
+		else if(UBlueprint* ActorBlueprint = Cast<UBlueprint>(ActorClass->ClassGeneratedBy))
+		{
+			UBlueprint* NewBlueprint = FKismetEditorUtilities::CreateBlueprint(ActorClass, MovieScene, *BlueprintName, EBlueprintType::BPTYPE_Normal, UBlueprint::StaticClass(), UBlueprintGeneratedClass::StaticClass(), NAME_Name);
+			if(NewBlueprint)
+			{
+				// if the duplicated BP has any skel mesh components, set them to single anim instance
+				AActor* CDO = CastChecked<AActor>(NewBlueprint->GeneratedClass->ClassDefaultObject);
+				if(CDO)
 				{
-					AssetPath /= Settings->AnimationSubDirectory;
+					TInlineComponentArray<USkeletalMeshComponent*> SkeletalMeshComponents;
+					CDO->GetComponents(SkeletalMeshComponents);	
+					for(USkeletalMeshComponent* SkeletalMeshComponent : SkeletalMeshComponents)
+					{
+						SkeletalMeshComponent->SetAnimationMode(EAnimationMode::AnimationSingleNode);
+						SkeletalMeshComponent->bEnableUpdateRateOptimizations = false;
+						SkeletalMeshComponent->MeshComponentUpdateFlag = EMeshComponentUpdateFlag::AlwaysTickPoseAndRefreshBones;
+						SkeletalMeshComponent->ForcedLodModel = 1;
+					}
 				}
 
-				FString AssetName = Settings->SequenceName.Len() > 0 ? Settings->SequenceName : TEXT("RecordedSequence");
-				AssetName += TEXT("_");
-				AssetName += ActorToRecord->GetActorLabel();
+				Blueprint = NewBlueprint;
+			}
+		}
+		else
+		{
+			UActorFactory* FactoryToUse = GEditor->FindActorFactoryForActorClass(ActorClass);
+			if(FactoryToUse)
+			{
+				Blueprint = FactoryToUse->CreateBlueprint(Actor, MovieScene, *BlueprintName);
+			}
+		}
+
+		if(Blueprint)
+		{
+			Guid = MovieScene->AddSpawnable(BlueprintName, Blueprint);
+		}
+
+		// now add tracks to record
+		if(Guid.IsValid())
+		{
+			// force set recording to record translations as we need this with no animation
+			ActorSettings.bRecordTransforms = true;
+
+			// grab components so we can track attachments
+			// don't include non-CDO here as they wont be part of our initial BP (duplicated above)
+			// we will catch these 'extra' components on the first tick
+			const bool bIncludeNonCDO = false;
+			SyncTrackedComponents(bIncludeNonCDO);
+
+			TInlineComponentArray<USceneComponent*> SceneComponents(ActorToRecord.Get());
 
-				Sequence = SequenceRecorderUtils::MakeNewAsset<UAnimSequence>(AssetPath, AssetName);
-				if(Sequence)
+			// check if components need recording
+			TInlineComponentArray<USceneComponent*> ValidSceneComponents;
+			for(TWeakObjectPtr<USceneComponent>& SceneComponent : TrackedComponents)
+			{
+				if(ValidComponent(SceneComponent.Get()))
 				{
-					// set skeleton
-					Sequence->SetSkeleton(SkeletalMeshComponent->SkeletalMesh->Skeleton);
-					bValidAutoGeneratedAnimation = true;
+					ValidSceneComponents.Add(SceneComponent.Get());
+
+					// add all parent components too
+					TArray<USceneComponent*> ParentComponents;
+					SceneComponent->GetParentComponents(ParentComponents);
+					for(USceneComponent* ParentComponent : ParentComponents)
+					{	
+						ValidSceneComponents.AddUnique(ParentComponent);
+					}
 				}
 			}
 
-			if(bValidSpecifiedAnimation || bValidAutoGeneratedAnimation)
+			TSharedPtr<FMovieSceneAnimationPropertyRecorder> FirstAnimRecorder = nullptr;
+			for(USceneComponent* SceneComponent : ValidSceneComponents)
 			{
-				FAnimationRecorderManager::Get().RecordAnimation(SkeletalMeshComponent, Sequence, AnimationSettings);
-				LastRecordedAnimation = Sequence;
-				return true;
+				TSharedPtr<FMovieSceneAnimationPropertyRecorder> AnimRecorder = StartRecordingComponentProperties(SceneComponent->GetFName(), SceneComponent, ActorToRecord.Get(), CurrentSequence, CurrentSequenceTime);
+				if(!FirstAnimRecorder.IsValid() && AnimRecorder.IsValid())
+				{
+					FirstAnimRecorder = AnimRecorder;
+				}
 			}
+
+			// add a spawn track so we actually get spawned
+			TSharedPtr<FMovieSceneSpawnPropertyRecorder> SpawnRecorder = MakeShareable(new FMovieSceneSpawnPropertyRecorder);
+			SpawnRecorder->CreateSection(ActorToRecord.Get(), MovieScene, Guid, CurrentSequenceTime, true);
+			PropertyRecorders.Add(SpawnRecorder);
+
+			// we need to create a transform track even if we arent recording transforms
+			TSharedPtr<FMovieScene3DTransformPropertyRecorder> TransformRecorder = MakeShareable(new FMovieScene3DTransformPropertyRecorder(FirstAnimRecorder));
+			TransformRecorder->CreateSection(ActorToRecord.Get(), MovieScene, Guid, CurrentSequenceTime, ActorSettings.bRecordTransforms);
+			PropertyRecorders.Add(TransformRecorder);
+
+			// add a visibility track 
+			TSharedPtr<FMovieSceneVisibilityPropertyRecorder> VisibilityRecorder = MakeShareable(new FMovieSceneVisibilityPropertyRecorder);
+			VisibilityRecorder->CreateSection(ActorToRecord.Get(), MovieScene, Guid, CurrentSequenceTime, ActorSettings.bRecordVisibility);
+			PropertyRecorders.Add(VisibilityRecorder);
+
+			// add an attachment track
+			TSharedPtr<FMovieScene3DAttachPropertyRecorder> AttachRecorder = MakeShareable(new FMovieScene3DAttachPropertyRecorder);
+			AttachRecorder->CreateSection(ActorToRecord.Get(), MovieScene, Guid, CurrentSequenceTime, true);
+			PropertyRecorders.Add(AttachRecorder);			
 		}
 	}
+}
 
-	return false;
+static FString GetUniquePosessableName(UMovieScene* MovieScene, const FString& BaseName)
+{
+	FString BlueprintName = BaseName;
+	auto DuplName = [&](FMovieScenePossessable& InPosessable)
+	{
+		return InPosessable.GetName() == BlueprintName;
+	};
+
+	int32 Index = 2;
+	FString UniqueString;
+	while (MovieScene->FindPossessable(DuplName))
+	{
+		BlueprintName.RemoveFromEnd(UniqueString);
+		UniqueString = FString::Printf(TEXT(" (%d)"), Index++);
+		BlueprintName += UniqueString;
+	}
+
+	return BlueprintName;
 }
 
-void UActorRecording::Tick(float DeltaSeconds)
+TSharedPtr<FMovieSceneAnimationPropertyRecorder> UActorRecording::StartRecordingComponentProperties(const FName& BindingName, USceneComponent* SceneComponent, UObject* BindingContext, ULevelSequence* CurrentSequence, float CurrentSequenceTime)
 {
-	if (ActorToRecord.IsValid())
+	// first create a possessable for this component to be controlled by
+	UMovieScene* OwnerMovieScene = CurrentSequence->GetMovieScene();
+
+	FString Name = GetUniquePosessableName(OwnerMovieScene, BindingName.ToString());
+
+	const FGuid PossessableGuid = OwnerMovieScene->AddPossessable(Name, SceneComponent->GetClass());
+
+	// Set up parent/child guids for possessables within spawnables
+	FMovieScenePossessable* ChildPossessable = OwnerMovieScene->FindPossessable(PossessableGuid);
+	if (ensure(ChildPossessable))
+	{
+		ChildPossessable->SetParent(Guid);
+	}
+
+	FMovieSceneSpawnable* ParentSpawnable = OwnerMovieScene->FindSpawnable(Guid);
+	if (ParentSpawnable)
+	{
+		ParentSpawnable->AddChildPossessable(PossessableGuid);
+	}
+
+	FLevelSequenceObjectReference ObjectReference(FUniqueObjectGuid(), BindingName.ToString());
+
+	CurrentSequence->BindPossessableObject(PossessableGuid, ObjectReference);
+
+	TSharedPtr<FMovieSceneAnimationPropertyRecorder> AnimationRecorder;
+	if(SceneComponent->IsA<USkeletalMeshComponent>())
+	{
+		AnimationRecorder = MakeShareable(new FMovieSceneAnimationPropertyRecorder(AnimationSettings, TargetAnimation.Get()));
+		AnimationRecorder->CreateSection(SceneComponent, OwnerMovieScene, PossessableGuid, CurrentSequenceTime, true);
+		PropertyRecorders.Add(AnimationRecorder);
+	}
+	else if(SceneComponent->IsA<UParticleSystemComponent>())
 	{
-		USkeletalMeshComponent* SkeletalMeshComponent = ActorToRecord.Get()->FindComponentByClass<USkeletalMeshComponent>();
-		if (SkeletalMeshComponent)
+		TSharedPtr<FMovieSceneParticleTrackPropertyRecorder> ParticleTrackRecorder = MakeShareable(new FMovieSceneParticleTrackPropertyRecorder());
+		ParticleTrackRecorder->CreateSection(SceneComponent, OwnerMovieScene, PossessableGuid, CurrentSequenceTime, true);
+		PropertyRecorders.Add(ParticleTrackRecorder);
+	}
+
+	// Dont record the root component transforms as this will be taken into account by the actor transform track
+	if(SceneComponent != SceneComponent->GetOwner()->GetRootComponent())
+	{
+		TSharedPtr<FMovieScene3DTransformPropertyRecorder> TransformRecorder = MakeShareable(new FMovieScene3DTransformPropertyRecorder);
+		TransformRecorder->CreateSection(SceneComponent, OwnerMovieScene, PossessableGuid, CurrentSequenceTime, true);
+		PropertyRecorders.Add(TransformRecorder);
+	}
+
+	TSharedPtr<FMovieSceneVisibilityPropertyRecorder> VisibilityRecorder = MakeShareable(new FMovieSceneVisibilityPropertyRecorder);
+	VisibilityRecorder->CreateSection(SceneComponent, OwnerMovieScene, PossessableGuid, CurrentSequenceTime, true);
+	PropertyRecorders.Add(VisibilityRecorder);	
+
+	return AnimationRecorder;
+}
+
+void UActorRecording::Tick(float DeltaSeconds, ULevelSequence* CurrentSequence, float CurrentSequenceTime)
+{
+	if (IsRecording())
+	{
+		// check our components to see if they have changed
+		static TArray<USceneComponent*> SceneComponents;
+		GetSceneComponents(SceneComponents);
+
+		if(TrackedComponents.Num() != SceneComponents.Num())
 		{
-			// re-force updates on as gameplay can sometimes turn these back off!
-			SkeletalMeshComponent->bEnableUpdateRateOptimizations = false;
-			SkeletalMeshComponent->MeshComponentUpdateFlag = EMeshComponentUpdateFlag::AlwaysTickPoseAndRefreshBones;
+			StartRecordingNewComponents(CurrentSequence, CurrentSequenceTime);
+		}
+
+		for(auto& PropertyRecorder : PropertyRecorders)
+		{
+			PropertyRecorder->Record(CurrentSequenceTime);
 		}
 	}
 }
 
-bool UActorRecording::StopRecording()
+bool UActorRecording::StopRecording(ULevelSequence* CurrentSequence)
 {
-	if (ActorToRecord.IsValid())
+	// stop property recorders
+	for(auto& PropertyRecorder : PropertyRecorders)
 	{
-		// @todo: animation recording: allow recording of multiple skeletal meshes per actor
-		USkeletalMeshComponent* SkeletalMeshComponent = ActorToRecord.Get()->FindComponentByClass<USkeletalMeshComponent>();
-		if (SkeletalMeshComponent)
-		{
-			// restore update flags
-			SkeletalMeshComponent->bEnableUpdateRateOptimizations = bEnableUpdateRateOptimizations;
-			SkeletalMeshComponent->MeshComponentUpdateFlag = MeshComponentUpdateFlag;
+		PropertyRecorder->FinalizeSection();
+	}
 
-			FAnimationRecorderManager::Get().StopRecordingAnimation(SkeletalMeshComponent);
-			return true;
+	PropertyRecorders.Empty();
+
+	if(bNewComponentAddedWhileRecording)
+	{
+		UMovieScene* MovieScene = CurrentSequence->GetMovieScene();
+		check(MovieScene);
+
+		FMovieSceneSpawnable* Spawnable = MovieScene->FindSpawnable(Guid);
+		if(Spawnable)
+		{
+			UBlueprint* Blueprint = CastChecked<UBlueprint>(Spawnable->GetClass()->ClassGeneratedBy);
+			FBlueprintEditorUtils::MarkBlueprintAsStructurallyModified(Blueprint);
 		}
 	}
 
-	return false;
+	return true;
 }
 
 bool UActorRecording::IsRecording() const
 {
-	if(ActorToRecord.IsValid())
+	return ActorToRecord.IsValid() && PropertyRecorders.Num() > 0;
+}
+
+static FName FindParentComponentOwnerClassName(USceneComponent* SceneComponent, UBlueprint* Blueprint)
+{
+	if(SceneComponent->AttachParent)
 	{
-		USkeletalMeshComponent* SkeletalMeshComponent = ActorToRecord.Get()->FindComponentByClass<USkeletalMeshComponent>();
-		if (SkeletalMeshComponent)
+		FName AttachName = SceneComponent->AttachParent->GetFName();
+
+		// see if we can find this component in the BP inheritance hierarchy
+		while(Blueprint)
 		{
-			return FAnimationRecorderManager::Get().IsRecording(SkeletalMeshComponent);
+			if(Blueprint->SimpleConstructionScript->FindSCSNode(AttachName) != nullptr)
+			{
+				return Blueprint->GetFName();
+			}
+
+			Blueprint = Cast<UBlueprint>(Blueprint->GeneratedClass->GetSuperClass()->ClassGeneratedBy);
 		}
 	}
 
-	return false;
+	return NAME_None;
+}
+
+void UActorRecording::StartRecordingNewComponents(ULevelSequence* CurrentSequence, float CurrentSequenceTime)
+{
+	if (ActorToRecord.IsValid())
+	{
+		// find the new component(s)
+		TArray<USceneComponent*> NewComponents;
+		TArray<USceneComponent*> SceneComponents;
+		GetSceneComponents(SceneComponents);
+		for(USceneComponent* SceneComponent : SceneComponents)
+		{
+			if(ValidComponent(SceneComponent))
+			{
+				TWeakObjectPtr<USceneComponent> WeakSceneComponent(SceneComponent);
+				int32 FoundIndex = TrackedComponents.Find(WeakSceneComponent);
+				if(FoundIndex == INDEX_NONE)
+				{
+					// new component!
+					NewComponents.Add(SceneComponent);
+				}
+			}
+		}
+
+		UMovieScene* MovieScene = CurrentSequence->GetMovieScene();
+		check(MovieScene);
+
+		FMovieSceneSpawnable* Spawnable = MovieScene->FindSpawnable(Guid);
+		check(Spawnable);
+
+		UBlueprint* Blueprint = CastChecked<UBlueprint>(Spawnable->GetClass()->ClassGeneratedBy);
+
+		for(USceneComponent* SceneComponent : NewComponents)
+		{
+			// new component, so we need to add this to our BP if it didn't come from SCS
+			FName NewName;
+			if(SceneComponent->CreationMethod != EComponentCreationMethod::SimpleConstructionScript)
+			{
+				NewName = *FString::Printf(TEXT("Dynamic%s"), *SceneComponent->GetFName().ToString());
+				USCS_Node* NewSCSNode = Blueprint->SimpleConstructionScript->CreateNode(SceneComponent->GetClass(), NewName); 
+				UEditorEngine::CopyPropertiesForUnrelatedObjects(SceneComponent, NewSCSNode->ComponentTemplate);
+				NewSCSNode->ComponentTemplate->CreationMethod = EComponentCreationMethod::SimpleConstructionScript;
+
+				// look for a similar attach parent in the current structure
+				bool bAdded = false;
+				if(SceneComponent->AttachParent != nullptr)
+				{
+					FName AttachName = SceneComponent->AttachParent->GetFName();
+					NewSCSNode->ParentComponentOrVariableName = AttachName;
+					NewSCSNode->bIsParentComponentNative = SceneComponent->AttachParent->CreationMethod == EComponentCreationMethod::Native;
+					NewSCSNode->ParentComponentOwnerClassName = FindParentComponentOwnerClassName(SceneComponent, Blueprint);
+
+					USCS_Node* ParentSCSNode = Blueprint->SimpleConstructionScript->FindSCSNode(AttachName);
+					if(ParentSCSNode)
+					{
+						ParentSCSNode->AddChildNode(NewSCSNode);
+						bAdded = true;
+					}
+				}
+			
+				// nothing found - add to root
+				if(!bAdded)
+				{
+					Blueprint->SimpleConstructionScript->AddNode(NewSCSNode);
+				}
+
+				if(SceneComponent->AttachSocketName != NAME_None)
+				{
+					NewSCSNode->AttachToName = SceneComponent->AttachSocketName;
+				}
+			}
+			else
+			{
+				NewName = SceneComponent->GetFName();
+			}
+
+			StartRecordingComponentProperties(NewName, SceneComponent, ActorToRecord.Get(), CurrentSequence, CurrentSequenceTime);
+
+			bNewComponentAddedWhileRecording = true;
+		}
+		
+		SyncTrackedComponents();
+	}
 }
diff --git a/Engine/Plugins/Editor/SequenceRecorder/Source/SequenceRecorder/Private/ActorRecording.h b/Engine/Plugins/Editor/SequenceRecorder/Source/SequenceRecorder/Private/ActorRecording.h
--- a/Engine/Plugins/Editor/SequenceRecorder/Source/SequenceRecorder/Private/ActorRecording.h
+++ b/Engine/Plugins/Editor/SequenceRecorder/Source/SequenceRecorder/Private/ActorRecording.h
@@ -15,8 +15,6 @@ struct FActorRecordingSettings
 	FActorRecordingSettings()
 		: bRecordTransforms(false)
 		, bRecordVisibility(true)
-		, bRecordNearbySpawnedActors(false)
-		, NearbyActorRecordingProximity(1000.0f)
 	{}
 
 public:
@@ -30,34 +28,48 @@ public:
 	/** Whether to record actor visibility. */
 	UPROPERTY(EditAnywhere, Category = "Actor Recording")
 	bool bRecordVisibility;
-
-	/** Whether to record nearby spawned actors. */
-	UPROPERTY(EditAnywhere, Category = "Actor Recording")
-	bool bRecordNearbySpawnedActors;
-
-	/** Proximity to currently recorded actors to record newly spawned actors. */
-	UPROPERTY(EditAnywhere, Category = "Actor Recording")
-	float NearbyActorRecordingProximity;
 };
 
 UCLASS(MinimalAPI, Transient)
 class UActorRecording : public UObject
 {
-	GENERATED_BODY()
+	GENERATED_UCLASS_BODY()
 
 public:
-	/** Start this queued recording */
-	bool StartRecording();
+	/** Check whether it is worth recording this actor - i.e. is it going to affect the end result of the sequence */
+	static bool IsRelevantForRecording(AActor* Actor);
 
-	/** Stop this recording. Has no effect if we are not currently recording */
-	bool StopRecording();
+	/** Start this queued recording. Sequence can be nullptr */
+	bool StartRecording(class ULevelSequence* CurrentSequence = nullptr, float CurrentSequenceTime = 0.0f);
+
+	/** Stop this recording. Has no effect if we are not currently recording. Sequence can be nullptr */
+	bool StopRecording(class ULevelSequence* CurrentSequence = nullptr);
 
 	/** Tick this recording */
-	void Tick(float DeltaSeconds);
+	void Tick(float DeltaSeconds, ULevelSequence* CurrentSequence = nullptr, float CurrentSequenceTime = 0.0f);
 
 	/** Whether we are currently recording */
 	bool IsRecording() const;
 
+private:
+	/** Check component valididty for recording */
+	bool ValidComponent(USceneComponent* SceneComponent) const;
+
+	/** Start recording actor properties to a sequence */
+	void StartRecordingActorProperties(ULevelSequence* CurrentSequence, float CurrentSequenceTime);
+
+	/** Start recording component properties to a sequence */
+	TSharedPtr<class FMovieSceneAnimationPropertyRecorder> StartRecordingComponentProperties(const FName& BindingName, USceneComponent* SceneComponent, UObject* BindingContext, ULevelSequence* CurrentSequence, float CurrentSequenceTime);
+
+	/** Start recording components that are added at runtime */
+	void StartRecordingNewComponents(ULevelSequence* CurrentSequence, float CurrentSequenceTime);
+
+	/** Helper function to grab all scene components in the actor's hierarchy */
+	void GetSceneComponents(TArray<USceneComponent*>& OutArray, bool bIncludeNonCDO = true);
+
+	/** Sync up tracked components with the actor */
+	void SyncTrackedComponents(bool bIncludeNonCDO = true);
+
 public:
 	/** The actor we want to record */
 	UPROPERTY(EditAnywhere, Category = "Actor Recording")
@@ -78,17 +90,25 @@ public:
 	UPROPERTY(EditAnywhere, Category = "Animation Recording")
 	FAnimationRecordingSettings AnimationSettings;
 
-	/** Our last recorded animation */
-	UPROPERTY(Transient)
-	TWeakObjectPtr<UAnimSequence> LastRecordedAnimation;
-
 	/** Guid that identifies our spawnable in a recorded sequence */
 	FGuid Guid;
 
+	/** Whether this actor recording was triggered from an acotr spawn */
+	bool bWasSpawnedPostRecord;
+
 private:
 	/** Used to store/restore update flag when recording */
 	EMeshComponentUpdateFlag::Type MeshComponentUpdateFlag;
 
 	/** Used to store/restore URO when recording */
 	bool bEnableUpdateRateOptimizations;
+
+	/** This actor's current set of property recorders */
+	TArray<TSharedPtr<class IMovieScenePropertyRecorder>> PropertyRecorders;
+
+	/** Track components to check if any have changed */
+	TArray<TWeakObjectPtr<USceneComponent>> TrackedComponents;
+
+	/** Flag to track whether we created new components */
+	bool bNewComponentAddedWhileRecording;
 };
\ No newline at end of file
diff --git a/Engine/Plugins/Editor/SequenceRecorder/Source/SequenceRecorder/Private/AnimationRecorder.cpp b/Engine/Plugins/Editor/SequenceRecorder/Source/SequenceRecorder/Private/AnimationRecorder.cpp
--- a/Engine/Plugins/Editor/SequenceRecorder/Source/SequenceRecorder/Private/AnimationRecorder.cpp
+++ b/Engine/Plugins/Editor/SequenceRecorder/Source/SequenceRecorder/Private/AnimationRecorder.cpp
@@ -171,6 +171,42 @@ bool FAnimationRecorder::TriggerRecordAnimation(USkeletalMeshComponent* Componen
 	return false;
 }
 
+/** Helper function to get space bases depending on master pose component */
+static void GetBoneTransforms(USkeletalMeshComponent* Component, TArray<FTransform>& BoneTransforms)
+{
+	const USkinnedMeshComponent* const MasterPoseComponentInst = Component->MasterPoseComponent.Get();
+	if(MasterPoseComponentInst)
+	{
+		const TArray<FTransform>& SpaceBases = MasterPoseComponentInst->GetSpaceBases();
+		BoneTransforms.Reset(BoneTransforms.Num());
+		BoneTransforms.AddUninitialized(SpaceBases.Num());
+		for(int32 BoneIndex = 0; BoneIndex < SpaceBases.Num(); BoneIndex++)
+		{
+			if(BoneIndex < Component->GetMasterBoneMap().Num())
+			{
+				int32 MasterBoneIndex = Component->GetMasterBoneMap()[BoneIndex];
+
+				// If ParentBoneIndex is valid, grab matrix from MasterPoseComponent.
+				if(MasterBoneIndex != INDEX_NONE && MasterBoneIndex < SpaceBases.Num())
+				{
+					BoneTransforms[BoneIndex] = SpaceBases[MasterBoneIndex];
+				}
+				else
+				{
+					BoneTransforms[BoneIndex] = FTransform::Identity;
+				}
+			}
+			else
+			{
+				BoneTransforms[BoneIndex] = FTransform::Identity;
+			}
+		}
+	}
+	else
+	{
+		BoneTransforms =  Component->GetSpaceBases();
+	}
+}
 
 void FAnimationRecorder::StartRecord(USkeletalMeshComponent* Component, UAnimSequence* InAnimationObject)
 {
@@ -181,7 +217,7 @@ void FAnimationRecorder::StartRecord(USkeletalMeshComponent* Component, UAnimSeq
 	AnimationObject->TrackToSkeletonMapTable.Empty();
 	AnimationObject->AnimationTrackNames.Empty();
 
-	PreviousSpacesBases = Component->GetSpaceBases();
+	GetBoneTransforms(Component, PreviousSpacesBases);
 	PreviousComponentToWorld = Component->ComponentToWorld;
 
 	LastFrame = 0;
@@ -367,10 +403,11 @@ void FAnimationRecorder::UpdateRecord(USkeletalMeshComponent* Component, float D
 		RecordNotifies(Component, Component->GetAnimInstance()->NotifyQueue.AnimNotifies, DeltaTime, TimePassed);
 	}
 
+	TArray<FTransform> SpaceBases;
+	GetBoneTransforms(Component, SpaceBases);
+
 	if (FramesRecorded < FramesToRecord)
 	{
-		const TArray<FTransform>& SpaceBases = Component->GetSpaceBases();
-
 		check(SpaceBases.Num() == PreviousSpacesBases.Num());
 
 		TArray<FTransform> BlendedSpaceBases;
@@ -404,7 +441,7 @@ void FAnimationRecorder::UpdateRecord(USkeletalMeshComponent* Component, float D
 	}
 
 	//save to current transform
-	PreviousSpacesBases = Component->GetSpaceBases();
+	PreviousSpacesBases = SpaceBases;
 	PreviousComponentToWorld = Component->ComponentToWorld;
 
 	// if we passed MaxFrame, just stop it
@@ -415,7 +452,7 @@ void FAnimationRecorder::UpdateRecord(USkeletalMeshComponent* Component, float D
 	}
 }
 
-void FAnimationRecorder::Record( USkeletalMeshComponent* Component, FTransform const& ComponentToWorld, TArray<FTransform> SpacesBases, int32 FrameToAdd )
+void FAnimationRecorder::Record( USkeletalMeshComponent* Component, FTransform const& ComponentToWorld, const TArray<FTransform>& SpacesBases, int32 FrameToAdd )
 {
 	if (ensure(AnimationObject))
 	{
diff --git a/Engine/Plugins/Editor/SequenceRecorder/Source/SequenceRecorder/Private/AnimationRecorder.h b/Engine/Plugins/Editor/SequenceRecorder/Source/SequenceRecorder/Private/AnimationRecorder.h
--- a/Engine/Plugins/Editor/SequenceRecorder/Source/SequenceRecorder/Private/AnimationRecorder.h
+++ b/Engine/Plugins/Editor/SequenceRecorder/Source/SequenceRecorder/Private/AnimationRecorder.h
@@ -58,7 +58,7 @@ public:
 	uint8 bAutoSaveAsset : 1;
 
 private:
-	void Record(USkeletalMeshComponent* Component, FTransform const& ComponentToWorld, TArray<FTransform> SpacesBases, int32 FrameToAdd);
+	void Record(USkeletalMeshComponent* Component, FTransform const& ComponentToWorld, const TArray<FTransform>& SpacesBases, int32 FrameToAdd);
 
 	void RecordNotifies(USkeletalMeshComponent* Component, const TArray<const struct FAnimNotifyEvent*>& AnimNotifies, float DeltaTime, float RecordTime);
 
diff --git a/Engine/Plugins/Editor/SequenceRecorder/Source/SequenceRecorder/Private/IMovieScenePropertyRecorder.h b/Engine/Plugins/Editor/SequenceRecorder/Source/SequenceRecorder/Private/IMovieScenePropertyRecorder.h
--- a/Engine/Plugins/Editor/SequenceRecorder/Source/SequenceRecorder/Private/IMovieScenePropertyRecorder.h
+++ b/Engine/Plugins/Editor/SequenceRecorder/Source/SequenceRecorder/Private/IMovieScenePropertyRecorder.h
@@ -5,7 +5,7 @@
 class IMovieScenePropertyRecorder
 {
 public:
-	virtual void CreateSection(class AActor* SourceActor, class UMovieScene* MovieScene, const FGuid& Guid, bool bRecord) = 0;
+	virtual void CreateSection(UObject* InObjectToRecord, class UMovieScene* MovieScene, const FGuid& Guid, float Time, bool bRecord) = 0;
 	virtual void FinalizeSection() = 0;
 	virtual void Record(float CurrentTime) = 0;
 };
\ No newline at end of file
diff --git a/Engine/Plugins/Editor/SequenceRecorder/Source/SequenceRecorder/Private/PropertyRecorders/MovieScene3DAttachPropertyRecorder.cpp b/Engine/Plugins/Editor/SequenceRecorder/Source/SequenceRecorder/Private/PropertyRecorders/MovieScene3DAttachPropertyRecorder.cpp
new file mode 100644
--- /dev/null
+++ b/Engine/Plugins/Editor/SequenceRecorder/Source/SequenceRecorder/Private/PropertyRecorders/MovieScene3DAttachPropertyRecorder.cpp
@@ -0,0 +1,72 @@
+// Copyright 1998-2015 Epic Games, Inc. All Rights Reserved.
+
+#include "SequenceRecorderPrivatePCH.h"
+#include "MovieScene3DAttachPropertyRecorder.h"
+#include "MovieScene3DAttachSection.h"
+#include "MovieScene3DAttachTrack.h"
+#include "MovieScene.h"
+#include "SequenceRecorderUtils.h"
+#include "SequenceRecorder.h"
+#include "ActorRecording.h"
+
+void FMovieScene3DAttachPropertyRecorder::CreateSection(UObject* InObjectToRecord, UMovieScene* InMovieScene, const FGuid& Guid, float Time, bool bRecord)
+{
+	ObjectGuid = Guid;
+	ActorToRecord = CastChecked<AActor>(InObjectToRecord);
+	MovieScene = InMovieScene;
+
+	bRecording = bRecord;
+
+	Record(Time);
+}
+
+void FMovieScene3DAttachPropertyRecorder::FinalizeSection()
+{
+	bRecording = false;
+}
+
+void FMovieScene3DAttachPropertyRecorder::Record(float CurrentTime)
+{
+	if(ActorToRecord.IsValid())
+	{
+		if(bRecording)
+		{
+			if(MovieSceneSection.IsValid())
+			{
+				MovieSceneSection->SetEndTime(CurrentTime);
+			}
+
+			// get attachment and check if the actor we are attached to is being recorded
+			FName SocketName;
+			FName ComponentName;
+			AActor* AttachedToActor = SequenceRecorderUtils::GetAttachment(ActorToRecord.Get(), SocketName, ComponentName);
+			UActorRecording* ActorRecording = FSequenceRecorder::Get().FindRecording(AttachedToActor);
+			if(AttachedToActor && ActorRecording)
+			{
+				// create the track if we haven't already
+				if(!AttachTrack.IsValid())
+				{
+					AttachTrack = MovieScene->AddTrack<UMovieScene3DAttachTrack>(ObjectGuid);
+				}
+
+				// check if we need a section or if the actor we are attached to has changed
+				if(!MovieSceneSection.IsValid() || AttachedToActor != ActorAttachedTo.Get())
+				{
+					MovieSceneSection = Cast<UMovieScene3DAttachSection>(AttachTrack->CreateNewSection());
+					AttachTrack->AddSection(*MovieSceneSection);
+					MovieSceneSection->SetStartTime(CurrentTime);
+					MovieSceneSection->SetConstraintId(ActorRecording->Guid);
+					MovieSceneSection->AttachSocketName = SocketName;
+					MovieSceneSection->AttachComponentName = ComponentName;
+				}
+
+				ActorAttachedTo = AttachedToActor;
+			}
+			else
+			{
+				// no attachment, so end the section recording if we have any
+				MovieSceneSection = nullptr;
+			}
+		}
+	}
+}
\ No newline at end of file
diff --git a/Engine/Plugins/Editor/SequenceRecorder/Source/SequenceRecorder/Private/PropertyRecorders/MovieScene3DAttachPropertyRecorder.h b/Engine/Plugins/Editor/SequenceRecorder/Source/SequenceRecorder/Private/PropertyRecorders/MovieScene3DAttachPropertyRecorder.h
new file mode 100644
--- /dev/null
+++ b/Engine/Plugins/Editor/SequenceRecorder/Source/SequenceRecorder/Private/PropertyRecorders/MovieScene3DAttachPropertyRecorder.h
@@ -0,0 +1,37 @@
+// Copyright 1998-2015 Epic Games, Inc. All Rights Reserved.
+
+#pragma once
+
+#include "IMovieScenePropertyRecorder.h"
+
+class FMovieScene3DAttachPropertyRecorder : public IMovieScenePropertyRecorder
+{
+public:
+	virtual ~FMovieScene3DAttachPropertyRecorder() {}
+
+	virtual void CreateSection(UObject* InObjectToRecord, class UMovieScene* InMovieScene, const FGuid& Guid, float Time, bool bRecord) override;
+	virtual void FinalizeSection() override;
+	virtual void Record(float CurrentTime) override;
+
+private:
+	/** Object to record from */
+	TLazyObjectPtr<class AActor> ActorToRecord;
+
+	/** Section to record to */
+	TWeakObjectPtr<class UMovieScene3DAttachSection> MovieSceneSection;
+
+	/** Track we are recording to */
+	TWeakObjectPtr<class UMovieScene3DAttachTrack> AttachTrack;
+
+	/** Movie scene we are recording to */
+	TWeakObjectPtr<class UMovieScene> MovieScene;
+
+	/** Track the actor we are attached to */
+	TLazyObjectPtr<class AActor> ActorAttachedTo;
+
+	/** Identifier of the object we are recording */
+	FGuid ObjectGuid;
+
+	/** Flag if we are actually recording or not */
+	bool bRecording;
+};
\ No newline at end of file
diff --git a/Engine/Plugins/Editor/SequenceRecorder/Source/SequenceRecorder/Private/PropertyRecorders/MovieScene3DTransformPropertyRecorder.cpp b/Engine/Plugins/Editor/SequenceRecorder/Source/SequenceRecorder/Private/PropertyRecorders/MovieScene3DTransformPropertyRecorder.cpp
--- a/Engine/Plugins/Editor/SequenceRecorder/Source/SequenceRecorder/Private/PropertyRecorders/MovieScene3DTransformPropertyRecorder.cpp
+++ b/Engine/Plugins/Editor/SequenceRecorder/Source/SequenceRecorder/Private/PropertyRecorders/MovieScene3DTransformPropertyRecorder.cpp
@@ -5,18 +5,16 @@
 #include "MovieScene.h"
 #include "MovieScene3DTransformSection.h"
 #include "MovieScene3DTransformTrack.h"
+#include "MovieSceneAnimationPropertyRecorder.h"
+#include "SequenceRecorderUtils.h"
 
-static const FName TransformTrackName = TEXT("Transform");
-
-void FMovieScene3DTransformPropertyRecorder::CreateSection(AActor* SourceActor, UMovieScene* MovieScene, const FGuid& Guid, bool bRecord) 
+void FMovieScene3DTransformPropertyRecorder::CreateSection(UObject* InObjectToRecord, UMovieScene* MovieScene, const FGuid& Guid, float Time, bool bRecord) 
 {
-	ActorToRecord = SourceActor;
+	ObjectToRecord = InObjectToRecord;
 
 	UMovieScene3DTransformTrack* TransformTrack = MovieScene->AddTrack<UMovieScene3DTransformTrack>(Guid);
 	if(TransformTrack)
 	{
-		TransformTrack->SetPropertyNameAndPath(TransformTrackName, TransformTrackName.ToString());
-
 		MovieSceneSection = Cast<UMovieScene3DTransformSection>(TransformTrack->CreateNewSection());
 
 		TransformTrack->AddSection(*MovieSceneSection);
@@ -35,28 +33,153 @@ void FMovieScene3DTransformPropertyRecorder::CreateSection(AActor* SourceActor,
 		MovieSceneSection->SetDefault(FTransformKey(EKey3DTransformChannel::Scale, EAxis::Y, 1.0f, bUnwindRotation));
 		MovieSceneSection->SetDefault(FTransformKey(EKey3DTransformChannel::Scale, EAxis::Z, 1.0f, bUnwindRotation));
 
-		MovieSceneSection->SetStartTime(0.0f);
+		MovieSceneSection->SetStartTime(Time);
 	}
 
 	bRecording = bRecord;
+
+	Record(Time);
+}
+
+static void WindRelativeRotators(float& Axis0, float& Axis1)
+{
+	float Diff = Axis0 - Axis1;
+	float AbsDiff = FMath::Abs(Diff);
+	if(AbsDiff > 180.0f)
+	{
+		Axis1 += 360.0f * FMath::Sign(Diff) * FMath::FloorToFloat((AbsDiff / 360.0f) + 0.5f);
+	}
 }
 
 void FMovieScene3DTransformPropertyRecorder::FinalizeSection()
 {
 	bRecording = false;
 
+	// if we have a valid animation recorder, we need to build our transforms from the animation
+	// so we properly synchronize our keyframes
+	if(AnimRecorder.IsValid())
+	{
+		check(BufferedTransforms.Num() == 0);
+
+		UAnimSequence* AnimSequence = AnimRecorder->GetSequence();
+		USkeletalMesh* SkeletalMesh = AnimRecorder->GetSkeletalMesh();
+		if(AnimSequence && SkeletalMesh)
+		{
+			// find the root bone
+			int32 RootIndex = INDEX_NONE;
+			USkeleton* AnimSkeleton = AnimSequence->GetSkeleton();
+			for (int32 TrackIndex = 0; TrackIndex < AnimSequence->RawAnimationData.Num(); ++TrackIndex)
+			{
+				// verify if this bone exists in skeleton
+				int32 BoneTreeIndex = AnimSequence->GetSkeletonIndexFromTrackIndex(TrackIndex);
+				if (BoneTreeIndex != INDEX_NONE)
+				{
+					int32 BoneIndex = AnimSkeleton->GetMeshBoneIndexFromSkeletonBoneIndex(SkeletalMesh, BoneTreeIndex);
+					int32 ParentIndex = SkeletalMesh->RefSkeleton.GetParentIndex(BoneIndex);
+					if(ParentIndex == INDEX_NONE)
+					{
+						// found root
+						RootIndex = BoneIndex;
+						break;
+					}
+				}
+			}
+
+			check(RootIndex != INDEX_NONE);
+
+			const float StartTime = MovieSceneSection->GetStartTime();
+
+			// we may need to offset the transform here if the animation was not recorded on the root component
+			FTransform InvComponentTransform = AnimRecorder->GetComponentTransform().Inverse();
+
+			FRawAnimSequenceTrack& RawTrack = AnimSequence->RawAnimationData[RootIndex];
+			const int32 KeyCount = FMath::Max(FMath::Max(RawTrack.PosKeys.Num(), RawTrack.RotKeys.Num()), RawTrack.ScaleKeys.Num());
+			for(int32 KeyIndex = 0; KeyIndex < KeyCount; KeyIndex++)
+			{
+				FTransform Transform;
+				if(RawTrack.PosKeys.IsValidIndex(KeyIndex))
+				{
+					Transform.SetTranslation(RawTrack.PosKeys[KeyIndex]);
+				}
+				else if(RawTrack.PosKeys.Num() > 0)
+				{
+					Transform.SetTranslation(RawTrack.PosKeys[0]);
+				}
+				
+				if(RawTrack.RotKeys.IsValidIndex(KeyIndex))
+				{
+					Transform.SetRotation(RawTrack.RotKeys[KeyIndex]);
+				}
+				else if(RawTrack.RotKeys.Num() > 0)
+				{
+					Transform.SetRotation(RawTrack.RotKeys[0]);
+				}
+
+				if(RawTrack.ScaleKeys.IsValidIndex(KeyIndex))
+				{
+					Transform.SetScale3D(RawTrack.ScaleKeys[KeyIndex]);
+				}
+				else if(RawTrack.ScaleKeys.Num() > 0)
+				{
+					Transform.SetScale3D(RawTrack.ScaleKeys[0]);
+				}
+
+				BufferedTransforms.Add(FBufferedTransformKey(InvComponentTransform * Transform, StartTime + AnimSequence->GetTimeAtFrame(KeyIndex)));
+			}
+		}
+	}
+
+	// Try to 're-wind' rotations that look like axis flips
+	// We need to do this as a post-process because the recorder cant reliably access 'wound' rotations:
+	// - Net quantize may use quaternions.
+	// - Scene components cache transforms as quaternions.
+	// - Gameplay is free to clamp/fmod rotations as it sees fit.
+	int32 TransformCount = BufferedTransforms.Num();
+	for(int32 TransformIndex = 0; TransformIndex < TransformCount - 1; TransformIndex++)
+	{
+		FRotator& Rotator = BufferedTransforms[TransformIndex].WoundRotation;
+		FRotator& NextRotator = BufferedTransforms[TransformIndex + 1].WoundRotation;
+
+		WindRelativeRotators(Rotator.Pitch, NextRotator.Pitch);
+		WindRelativeRotators(Rotator.Yaw, NextRotator.Yaw);
+		WindRelativeRotators(Rotator.Roll, NextRotator.Roll);
+	}
+
+	// add buffered transforms
+	const bool bUnwindRotation = false;
+	for(const FBufferedTransformKey& BufferedTransform : BufferedTransforms)
+	{
+		const FVector Translation = BufferedTransform.Transform.GetTranslation();
+		const FVector Rotation = BufferedTransform.WoundRotation.Euler();
+		const FVector Scale = BufferedTransform.Transform.GetScale3D();
+
+		MovieSceneSection->AddKey(BufferedTransform.KeyTime, FTransformKey(EKey3DTransformChannel::Translation, EAxis::X, Translation.X, bUnwindRotation), EMovieSceneKeyInterpolation::Linear);
+		MovieSceneSection->AddKey(BufferedTransform.KeyTime, FTransformKey(EKey3DTransformChannel::Translation, EAxis::Y, Translation.Y, bUnwindRotation), EMovieSceneKeyInterpolation::Linear);
+		MovieSceneSection->AddKey(BufferedTransform.KeyTime, FTransformKey(EKey3DTransformChannel::Translation, EAxis::Z, Translation.Z, bUnwindRotation), EMovieSceneKeyInterpolation::Linear);
+
+		MovieSceneSection->AddKey(BufferedTransform.KeyTime, FTransformKey(EKey3DTransformChannel::Rotation, EAxis::X, Rotation.X, bUnwindRotation), EMovieSceneKeyInterpolation::Linear);
+		MovieSceneSection->AddKey(BufferedTransform.KeyTime, FTransformKey(EKey3DTransformChannel::Rotation, EAxis::Y, Rotation.Y, bUnwindRotation), EMovieSceneKeyInterpolation::Linear);
+		MovieSceneSection->AddKey(BufferedTransform.KeyTime, FTransformKey(EKey3DTransformChannel::Rotation, EAxis::Z, Rotation.Z, bUnwindRotation), EMovieSceneKeyInterpolation::Linear);
+
+		MovieSceneSection->AddKey(BufferedTransform.KeyTime, FTransformKey(EKey3DTransformChannel::Scale, EAxis::X, Scale.X, bUnwindRotation), EMovieSceneKeyInterpolation::Linear);
+		MovieSceneSection->AddKey(BufferedTransform.KeyTime, FTransformKey(EKey3DTransformChannel::Scale, EAxis::Y, Scale.Y, bUnwindRotation), EMovieSceneKeyInterpolation::Linear);
+		MovieSceneSection->AddKey(BufferedTransform.KeyTime, FTransformKey(EKey3DTransformChannel::Scale, EAxis::Z, Scale.Z, bUnwindRotation), EMovieSceneKeyInterpolation::Linear);
+	}
+
+	BufferedTransforms.Empty();
+
 	// now remove linear keys
 	TPair<FRichCurve*, float> CurvesAndTolerances[] =
 	{
-		TPairInitializer<FRichCurve*, float>(&MovieSceneSection->GetTranslationCurve(EAxis::X), 0.1f),
-		TPairInitializer<FRichCurve*, float>(&MovieSceneSection->GetTranslationCurve(EAxis::Y), 0.1f),
-		TPairInitializer<FRichCurve*, float>(&MovieSceneSection->GetTranslationCurve(EAxis::Z), 0.1f),
-		TPairInitializer<FRichCurve*, float>(&MovieSceneSection->GetRotationCurve(EAxis::X), 0.25f),
-		TPairInitializer<FRichCurve*, float>(&MovieSceneSection->GetRotationCurve(EAxis::Y), 0.25f),
-		TPairInitializer<FRichCurve*, float>(&MovieSceneSection->GetRotationCurve(EAxis::Z), 0.25f),
-		TPairInitializer<FRichCurve*, float>(&MovieSceneSection->GetScaleCurve(EAxis::X), 0.1f),
-		TPairInitializer<FRichCurve*, float>(&MovieSceneSection->GetScaleCurve(EAxis::Y), 0.1f),
-		TPairInitializer<FRichCurve*, float>(&MovieSceneSection->GetScaleCurve(EAxis::Z), 0.1f),
+		TPairInitializer<FRichCurve*, float>(&MovieSceneSection->GetTranslationCurve(EAxis::X), KINDA_SMALL_NUMBER),
+		TPairInitializer<FRichCurve*, float>(&MovieSceneSection->GetTranslationCurve(EAxis::Y), KINDA_SMALL_NUMBER),
+		TPairInitializer<FRichCurve*, float>(&MovieSceneSection->GetTranslationCurve(EAxis::Z), KINDA_SMALL_NUMBER),
+		TPairInitializer<FRichCurve*, float>(&MovieSceneSection->GetRotationCurve(EAxis::X), KINDA_SMALL_NUMBER),
+		TPairInitializer<FRichCurve*, float>(&MovieSceneSection->GetRotationCurve(EAxis::Y), KINDA_SMALL_NUMBER),
+		TPairInitializer<FRichCurve*, float>(&MovieSceneSection->GetRotationCurve(EAxis::Z), KINDA_SMALL_NUMBER),
+		TPairInitializer<FRichCurve*, float>(&MovieSceneSection->GetScaleCurve(EAxis::X), KINDA_SMALL_NUMBER),
+		TPairInitializer<FRichCurve*, float>(&MovieSceneSection->GetScaleCurve(EAxis::Y), KINDA_SMALL_NUMBER),
+		TPairInitializer<FRichCurve*, float>(&MovieSceneSection->GetScaleCurve(EAxis::Z), KINDA_SMALL_NUMBER),
 	};
 
 	for(TPair<FRichCurve*, float>& CurveAndTolerance : CurvesAndTolerances)
@@ -67,32 +190,43 @@ void FMovieScene3DTransformPropertyRecorder::FinalizeSection()
 
 void FMovieScene3DTransformPropertyRecorder::Record(float CurrentTime)
 {
-	if(ActorToRecord.IsValid())
+	if(ObjectToRecord.IsValid())
 	{
+		if(USceneComponent* SceneComponent = Cast<USceneComponent>(ObjectToRecord.Get()))
+		{
+			// dont record non-registered scene components
+			if(!SceneComponent->IsRegistered())
+			{
+				return;
+			}
+		}
+
 		MovieSceneSection->SetEndTime(CurrentTime);
 
 		if(bRecording)
 		{
-			const bool bUnwindRotation = false;
-
-			USkeletalMeshComponent* SkeletalMeshComponent = ActorToRecord->FindComponentByClass<USkeletalMeshComponent>();
-
-			const FTransform Transform = SkeletalMeshComponent ? SkeletalMeshComponent->ComponentToWorld : ActorToRecord->GetTransform();
-			const FVector Translation = Transform.GetTranslation();
-			const FVector Rotation = Transform.GetRotation().Rotator().Euler();
-			const FVector Scale = Transform.GetScale3D();
-
-			MovieSceneSection->AddKey(CurrentTime, FTransformKey(EKey3DTransformChannel::Translation, EAxis::X, Translation.X, bUnwindRotation), EMovieSceneKeyInterpolation::Auto);
-			MovieSceneSection->AddKey(CurrentTime, FTransformKey(EKey3DTransformChannel::Translation, EAxis::Y, Translation.Y, bUnwindRotation), EMovieSceneKeyInterpolation::Auto);
-			MovieSceneSection->AddKey(CurrentTime, FTransformKey(EKey3DTransformChannel::Translation, EAxis::Z, Translation.Z, bUnwindRotation), EMovieSceneKeyInterpolation::Auto);
-
-			MovieSceneSection->AddKey(CurrentTime, FTransformKey(EKey3DTransformChannel::Rotation, EAxis::X, Rotation.X, bUnwindRotation), EMovieSceneKeyInterpolation::Auto);
-			MovieSceneSection->AddKey(CurrentTime, FTransformKey(EKey3DTransformChannel::Rotation, EAxis::Y, Rotation.Y, bUnwindRotation), EMovieSceneKeyInterpolation::Auto);
-			MovieSceneSection->AddKey(CurrentTime, FTransformKey(EKey3DTransformChannel::Rotation, EAxis::Z, Rotation.Z, bUnwindRotation), EMovieSceneKeyInterpolation::Auto);
-
-			MovieSceneSection->AddKey(CurrentTime, FTransformKey(EKey3DTransformChannel::Scale, EAxis::X, Scale.X, bUnwindRotation), EMovieSceneKeyInterpolation::Auto);
-			MovieSceneSection->AddKey(CurrentTime, FTransformKey(EKey3DTransformChannel::Scale, EAxis::Y, Scale.Y, bUnwindRotation), EMovieSceneKeyInterpolation::Auto);
-			MovieSceneSection->AddKey(CurrentTime, FTransformKey(EKey3DTransformChannel::Scale, EAxis::Z, Scale.Z, bUnwindRotation), EMovieSceneKeyInterpolation::Auto);
+			// don't record from the transform of the component/actor if we are synchronizing with an animation
+			if(!AnimRecorder.IsValid())
+			{
+				FTransform Transform;
+				if(USceneComponent* SceneComponent = Cast<USceneComponent>(ObjectToRecord.Get()))
+				{
+					Transform = SceneComponent->GetRelativeTransform();
+				}
+				else if(AActor* Actor = Cast<AActor>(ObjectToRecord.Get()))
+				{
+					FName SocketName, ComponentName;
+					if(SequenceRecorderUtils::GetAttachment(Actor, SocketName, ComponentName) != nullptr && Actor->GetRootComponent() != nullptr)
+					{
+						Transform = Actor->GetRootComponent()->GetRelativeTransform();
+					}
+					else
+					{
+						Transform = Actor->GetTransform();
+					}
+				}
+				BufferedTransforms.Add(FBufferedTransformKey(Transform, CurrentTime));
+			}
 		}
 	}
 }
\ No newline at end of file
diff --git a/Engine/Plugins/Editor/SequenceRecorder/Source/SequenceRecorder/Private/PropertyRecorders/MovieScene3DTransformPropertyRecorder.h b/Engine/Plugins/Editor/SequenceRecorder/Source/SequenceRecorder/Private/PropertyRecorders/MovieScene3DTransformPropertyRecorder.h
--- a/Engine/Plugins/Editor/SequenceRecorder/Source/SequenceRecorder/Private/PropertyRecorders/MovieScene3DTransformPropertyRecorder.h
+++ b/Engine/Plugins/Editor/SequenceRecorder/Source/SequenceRecorder/Private/PropertyRecorders/MovieScene3DTransformPropertyRecorder.h
@@ -4,19 +4,47 @@
 
 #include "IMovieScenePropertyRecorder.h"
 
+/** Structure used to buffer up transform keys. Keys are inserted into tracks in FinalizeSection() */
+struct FBufferedTransformKey
+{
+	FBufferedTransformKey(const FTransform& InTransform, float InKeyTime)
+		: Transform(InTransform)
+		, WoundRotation(InTransform.Rotator())
+		, KeyTime(InKeyTime)
+	{
+	}
+
+	FTransform Transform;
+	FRotator WoundRotation;
+	float KeyTime;
+};
+
 class FMovieScene3DTransformPropertyRecorder : public IMovieScenePropertyRecorder
 {
 public:
+	FMovieScene3DTransformPropertyRecorder(TSharedPtr<class FMovieSceneAnimationPropertyRecorder> InAnimRecorder = nullptr) 
+		: AnimRecorder(InAnimRecorder)
+	{}
+
 	virtual ~FMovieScene3DTransformPropertyRecorder() {}
 
-	virtual void CreateSection(class AActor* SourceActor, class UMovieScene* MovieScene, const FGuid& Guid, bool bRecord) override;
+	virtual void CreateSection(UObject* InObjectToRecord, class UMovieScene* MovieScene, const FGuid& Guid, float Time, bool bRecord) override;
 	virtual void FinalizeSection() override;
 	virtual void Record(float CurrentTime) override;
 
 private:
-	TLazyObjectPtr<class AActor> ActorToRecord;
+	/** Object to record from */
+	TLazyObjectPtr<UObject> ObjectToRecord;
 
+	/** Section to record to */
 	TWeakObjectPtr<class UMovieScene3DTransformSection> MovieSceneSection;
 
+	/** Buffer of transform keys. Keys are inserted into tracks in FinalizeSection() */
+	TArray<FBufferedTransformKey> BufferedTransforms;
+
+	/** Flag if we are actually recording or not */
 	bool bRecording;
+
+	/** Animation recorder we use to sync our transforms to */
+	TSharedPtr<class FMovieSceneAnimationPropertyRecorder> AnimRecorder;
 };
\ No newline at end of file
diff --git a/Engine/Plugins/Editor/SequenceRecorder/Source/SequenceRecorder/Private/PropertyRecorders/MovieSceneAnimationPropertyRecorder.cpp b/Engine/Plugins/Editor/SequenceRecorder/Source/SequenceRecorder/Private/PropertyRecorders/MovieSceneAnimationPropertyRecorder.cpp
new file mode 100644
--- /dev/null
+++ b/Engine/Plugins/Editor/SequenceRecorder/Source/SequenceRecorder/Private/PropertyRecorders/MovieSceneAnimationPropertyRecorder.cpp
@@ -0,0 +1,125 @@
+// Copyright 1998-2015 Epic Games, Inc. All Rights Reserved.
+
+#include "SequenceRecorderPrivatePCH.h"
+#include "MovieSceneAnimationPropertyRecorder.h"
+#include "MovieSceneSkeletalAnimationTrack.h"
+#include "MovieSceneSkeletalAnimationSection.h"
+#include "MovieScene.h"
+#include "SequenceRecorderUtils.h"
+#include "SequenceRecorderSettings.h"
+
+FMovieSceneAnimationPropertyRecorder::FMovieSceneAnimationPropertyRecorder(FAnimationRecordingSettings& InAnimationSettings, UAnimSequence* InSpecifiedSequence)
+	: Sequence(InSpecifiedSequence)
+	, AnimationSettings(InAnimationSettings)
+{
+}
+
+void FMovieSceneAnimationPropertyRecorder::CreateSection(UObject* InObjectToRecord, UMovieScene* MovieScene, const FGuid& Guid, float Time, bool bRecord)
+{
+	ObjectToRecord = InObjectToRecord;
+
+	AActor* Actor = nullptr;
+	SkeletalMeshComponent = Cast<USkeletalMeshComponent>(InObjectToRecord);
+	if(!SkeletalMeshComponent.IsValid())
+	{
+		Actor = Cast<AActor>(InObjectToRecord);
+		if(Actor != nullptr)
+		{
+			SkeletalMeshComponent = Actor->FindComponentByClass<USkeletalMeshComponent>();
+		}
+	}
+	else
+	{
+		Actor = SkeletalMeshComponent->GetOwner();
+	}
+
+	if(SkeletalMeshComponent.IsValid())
+	{
+		SkeletalMesh = SkeletalMeshComponent->SkeletalMesh;
+
+		// turn off URO and make sure we always update even if out of view
+		bEnableUpdateRateOptimizations = SkeletalMeshComponent->bEnableUpdateRateOptimizations;
+		MeshComponentUpdateFlag = SkeletalMeshComponent->MeshComponentUpdateFlag;
+
+		SkeletalMeshComponent->bEnableUpdateRateOptimizations = false;
+		SkeletalMeshComponent->MeshComponentUpdateFlag = EMeshComponentUpdateFlag::AlwaysTickPoseAndRefreshBones;
+
+		ComponentTransform = SkeletalMeshComponent->GetComponentToWorld().GetRelativeTransform(SkeletalMeshComponent->GetOwner()->GetTransform());
+
+		if(!Sequence.IsValid())
+		{
+			// build an asset path
+			const USequenceRecorderSettings* Settings = GetDefault<USequenceRecorderSettings>();
+
+			FString AssetPath(TEXT("/Game"));
+			AssetPath /= Settings->SequenceRecordingBasePath.Path;
+			if(Settings->AnimationSubDirectory.Len() > 0)
+			{
+				AssetPath /= Settings->AnimationSubDirectory;
+			}
+
+			FString AssetName = Settings->SequenceName.Len() > 0 ? Settings->SequenceName : TEXT("RecordedSequence");
+			AssetName += TEXT("_");
+			AssetName += Actor->GetActorLabel();
+
+			Sequence = SequenceRecorderUtils::MakeNewAsset<UAnimSequence>(AssetPath, AssetName);
+			if(Sequence.IsValid())
+			{
+				// set skeleton
+				Sequence->SetSkeleton(SkeletalMeshComponent->SkeletalMesh->Skeleton);
+			}
+		}
+
+		if(Sequence.IsValid())
+		{
+			FAnimationRecorderManager::Get().RecordAnimation(SkeletalMeshComponent.Get(), Sequence.Get(), AnimationSettings);
+
+			if(MovieScene)
+			{
+				UMovieSceneSkeletalAnimationTrack* AnimTrack = MovieScene->AddTrack<UMovieSceneSkeletalAnimationTrack>(Guid);
+				if(AnimTrack)
+				{
+					AnimTrack->AddNewAnimation(Time, Sequence.Get());
+					MovieSceneSection = Cast<UMovieSceneSkeletalAnimationSection>(AnimTrack->GetAllSections()[0]);
+				}
+			}
+		}
+	}
+
+	bRecording = bRecord;
+}
+
+void FMovieSceneAnimationPropertyRecorder::FinalizeSection()
+{
+	bRecording = false;
+
+	if(SkeletalMeshComponent.IsValid())
+	{
+		// restore update flags
+		SkeletalMeshComponent->bEnableUpdateRateOptimizations = bEnableUpdateRateOptimizations;
+		SkeletalMeshComponent->MeshComponentUpdateFlag = MeshComponentUpdateFlag;
+
+		FAnimationRecorderManager::Get().StopRecordingAnimation(SkeletalMeshComponent.Get());
+	}
+
+	if(MovieSceneSection.IsValid() && Sequence.IsValid())
+	{
+		MovieSceneSection->SetEndTime(MovieSceneSection->GetStartTime() + Sequence->GetPlayLength());
+
+		// enable root motion on the animation
+		Sequence->bEnableRootMotion = true;
+		Sequence->RootMotionRootLock = ERootMotionRootLock::Zero;
+	}
+}
+
+void FMovieSceneAnimationPropertyRecorder::Record(float CurrentTime)
+{
+	// The animation recorder does most of the work here
+
+	if(SkeletalMeshComponent.IsValid())
+	{
+		// re-force updates on as gameplay can sometimes turn these back off!
+		SkeletalMeshComponent->bEnableUpdateRateOptimizations = false;
+		SkeletalMeshComponent->MeshComponentUpdateFlag = EMeshComponentUpdateFlag::AlwaysTickPoseAndRefreshBones;
+	}
+}
\ No newline at end of file
diff --git a/Engine/Plugins/Editor/SequenceRecorder/Source/SequenceRecorder/Private/PropertyRecorders/MovieSceneAnimationPropertyRecorder.h b/Engine/Plugins/Editor/SequenceRecorder/Source/SequenceRecorder/Private/PropertyRecorders/MovieSceneAnimationPropertyRecorder.h
new file mode 100644
--- /dev/null
+++ b/Engine/Plugins/Editor/SequenceRecorder/Source/SequenceRecorder/Private/PropertyRecorders/MovieSceneAnimationPropertyRecorder.h
@@ -0,0 +1,49 @@
+// Copyright 1998-2015 Epic Games, Inc. All Rights Reserved.
+
+#pragma once
+
+#include "IMovieScenePropertyRecorder.h"
+
+class FMovieSceneAnimationPropertyRecorder : public IMovieScenePropertyRecorder
+{
+public:
+	FMovieSceneAnimationPropertyRecorder(FAnimationRecordingSettings& InAnimationSettings, UAnimSequence* InSpecifiedSequence);
+	virtual ~FMovieSceneAnimationPropertyRecorder() {}
+
+	virtual void CreateSection(UObject* InObjectToRecord, class UMovieScene* MovieScene, const FGuid& Guid, float Time, bool bRecord) override;
+	virtual void FinalizeSection() override;
+	virtual void Record(float CurrentTime) override;
+
+	UAnimSequence* GetSequence() const { return Sequence.Get(); }
+
+	USkeletalMesh* GetSkeletalMesh() const { return SkeletalMesh.Get(); }
+
+	const FTransform& GetComponentTransform() const { return ComponentTransform; }
+
+private:
+	/** Object to record from */
+	TLazyObjectPtr<UObject> ObjectToRecord;
+
+	/** Section to record to */
+	TWeakObjectPtr<class UMovieSceneSkeletalAnimationSection> MovieSceneSection;
+
+	TWeakObjectPtr<class UAnimSequence> Sequence;
+
+	TWeakObjectPtr<class USkeletalMeshComponent> SkeletalMeshComponent;
+
+	TWeakObjectPtr<class USkeletalMesh> SkeletalMesh;
+
+	/** Local transform of the component we are recording */
+	FTransform ComponentTransform;
+
+	FAnimationRecordingSettings AnimationSettings;
+
+	/** Used to store/restore update flag when recording */
+	EMeshComponentUpdateFlag::Type MeshComponentUpdateFlag;
+
+	/** Used to store/restore URO when recording */
+	bool bEnableUpdateRateOptimizations;
+
+	/** Flag if we are actually recording or not */
+	bool bRecording;
+};
\ No newline at end of file
diff --git a/Engine/Plugins/Editor/SequenceRecorder/Source/SequenceRecorder/Private/PropertyRecorders/MovieSceneParticleTrackPropertyRecorder.cpp b/Engine/Plugins/Editor/SequenceRecorder/Source/SequenceRecorder/Private/PropertyRecorders/MovieSceneParticleTrackPropertyRecorder.cpp
new file mode 100644
--- /dev/null
+++ b/Engine/Plugins/Editor/SequenceRecorder/Source/SequenceRecorder/Private/PropertyRecorders/MovieSceneParticleTrackPropertyRecorder.cpp
@@ -0,0 +1,99 @@
+// Copyright 1998-2015 Epic Games, Inc. All Rights Reserved.
+
+#include "SequenceRecorderPrivatePCH.h"
+#include "MovieSceneParticleTrackPropertyRecorder.h"
+#include "MovieSceneParticleSection.h"
+#include "MovieSceneParticleTrack.h"
+#include "MovieScene.h"
+#include "Particles/ParticleSystemComponent.h"
+
+FMovieSceneParticleTrackPropertyRecorder::~FMovieSceneParticleTrackPropertyRecorder()
+{
+	if(DelegateProxy.IsValid())
+	{
+		DelegateProxy->PropertyRecorder = nullptr;
+		DelegateProxy->RemoveFromRoot();
+		DelegateProxy.Reset();
+	}
+}
+
+void FMovieSceneParticleTrackPropertyRecorder::CreateSection(UObject* InObjectToRecord, UMovieScene* MovieScene, const FGuid& Guid, float Time, bool bRecord)
+{
+	SystemToRecord = CastChecked<UParticleSystemComponent>(InObjectToRecord);
+
+	UMovieSceneParticleTrack* ParticleTrack = MovieScene->AddTrack<UMovieSceneParticleTrack>(Guid);
+	if(ParticleTrack)
+	{
+		MovieSceneSection = Cast<UMovieSceneParticleSection>(ParticleTrack->CreateNewSection());
+
+		ParticleTrack->AddSection(*MovieSceneSection);
+
+		MovieSceneSection->SetStartTime(Time);
+
+		bWasTriggered = false;
+
+		DelegateProxy = NewObject<UMovieSceneParticleTrackPropertyRecorder>();
+		DelegateProxy->PropertyRecorder = this;
+		DelegateProxy->AddToRoot();
+		UParticleSystemComponent::OnSystemPreActivationChange.AddUObject(DelegateProxy.Get(), &UMovieSceneParticleTrackPropertyRecorder::OnTriggered);
+	}
+
+	bRecording = bRecord;
+
+	PreviousState = EParticleKey::Deactivate;
+
+	Record(Time);
+}
+
+void FMovieSceneParticleTrackPropertyRecorder::FinalizeSection()
+{
+	bRecording = false;
+}
+
+void FMovieSceneParticleTrackPropertyRecorder::Record(float CurrentTime)
+{
+	if(SystemToRecord.IsValid())
+	{
+		MovieSceneSection->SetEndTime(CurrentTime);
+
+		if(bRecording)
+		{
+			EParticleKey::Type NewState = EParticleKey::Deactivate;
+			if(SystemToRecord->IsRegistered() && SystemToRecord->IsActive() && !SystemToRecord->bWasDeactivated)
+			{
+				if(bWasTriggered)
+				{
+					NewState = EParticleKey::Trigger;
+					bWasTriggered = false;
+				}
+				else
+				{
+					NewState = EParticleKey::Activate;
+				}
+			}
+			else
+			{
+				NewState = EParticleKey::Deactivate;
+			}
+
+			if(NewState != PreviousState)
+			{
+				MovieSceneSection->AddKey(CurrentTime, NewState);
+			}
+
+			if(NewState == EParticleKey::Trigger)
+			{
+				NewState = EParticleKey::Activate;
+			}
+			PreviousState = NewState;
+		}
+	}
+}
+
+void UMovieSceneParticleTrackPropertyRecorder::OnTriggered(UParticleSystemComponent* Component, bool bActivating)
+{ 
+	if(PropertyRecorder && PropertyRecorder->SystemToRecord.Get() == Component)
+	{
+		PropertyRecorder->bWasTriggered = bActivating;
+	}
+}
\ No newline at end of file
diff --git a/Engine/Plugins/Editor/SequenceRecorder/Source/SequenceRecorder/Private/PropertyRecorders/MovieSceneParticleTrackPropertyRecorder.h b/Engine/Plugins/Editor/SequenceRecorder/Source/SequenceRecorder/Private/PropertyRecorders/MovieSceneParticleTrackPropertyRecorder.h
new file mode 100644
--- /dev/null
+++ b/Engine/Plugins/Editor/SequenceRecorder/Source/SequenceRecorder/Private/PropertyRecorders/MovieSceneParticleTrackPropertyRecorder.h
@@ -0,0 +1,54 @@
+// Copyright 1998-2015 Epic Games, Inc. All Rights Reserved.
+
+#pragma once
+
+#include "IMovieScenePropertyRecorder.h"
+#include "MovieSceneParticleSection.h"
+#include "MovieSceneParticleTrackPropertyRecorder.generated.h"
+
+struct FMovieSceneParticleTrackPropertyRecorder;
+
+UCLASS(MinimalAPI, NotBlueprintable)
+class UMovieSceneParticleTrackPropertyRecorder : public UObject
+{
+	GENERATED_BODY()
+
+public:
+	UMovieSceneParticleTrackPropertyRecorder()
+		: PropertyRecorder(nullptr)
+	{
+	}
+
+	UFUNCTION()
+	void OnTriggered(UParticleSystemComponent* Component, bool bActivating);
+
+	FMovieSceneParticleTrackPropertyRecorder* PropertyRecorder;
+};
+
+struct FMovieSceneParticleTrackPropertyRecorder : public IMovieScenePropertyRecorder
+{
+public:
+	friend class UMovieSceneParticleTrackPropertyRecorder;
+
+	virtual ~FMovieSceneParticleTrackPropertyRecorder();
+
+	virtual void CreateSection(UObject* InObjectToRecord, class UMovieScene* MovieScene, const FGuid& Guid, float Time, bool bRecord) override;
+	virtual void FinalizeSection() override;
+	virtual void Record(float CurrentTime) override;
+
+private:
+	/** Object to record from */
+	TLazyObjectPtr<class UParticleSystemComponent> SystemToRecord;
+
+	/** Section to record to */
+	TWeakObjectPtr<class UMovieSceneParticleSection> MovieSceneSection;
+
+	/** Flag if we are actually recording or not */
+	bool bRecording;
+
+	bool bWasTriggered;
+
+	EParticleKey::Type PreviousState;
+
+	TWeakObjectPtr<UMovieSceneParticleTrackPropertyRecorder> DelegateProxy;
+};
diff --git a/Engine/Plugins/Editor/SequenceRecorder/Source/SequenceRecorder/Private/PropertyRecorders/MovieSceneSpawnPropertyRecorder.cpp b/Engine/Plugins/Editor/SequenceRecorder/Source/SequenceRecorder/Private/PropertyRecorders/MovieSceneSpawnPropertyRecorder.cpp
new file mode 100644
--- /dev/null
+++ b/Engine/Plugins/Editor/SequenceRecorder/Source/SequenceRecorder/Private/PropertyRecorders/MovieSceneSpawnPropertyRecorder.cpp
@@ -0,0 +1,61 @@
+// Copyright 1998-2015 Epic Games, Inc. All Rights Reserved.
+
+#include "SequenceRecorderPrivatePCH.h"
+#include "MovieSceneSpawnPropertyRecorder.h"
+#include "MovieSceneBoolSection.h"
+#include "MovieSceneSpawnTrack.h"
+#include "MovieScene.h"
+
+void FMovieSceneSpawnPropertyRecorder::CreateSection(UObject* InObjectToRecord, UMovieScene* MovieScene, const FGuid& Guid, float Time, bool bRecord)
+{
+	ObjectToRecord = InObjectToRecord;
+
+	UMovieSceneSpawnTrack* SpawnTrack = MovieScene->AddTrack<UMovieSceneSpawnTrack>(Guid);
+	if(SpawnTrack)
+	{
+		MovieSceneSection = Cast<UMovieSceneBoolSection>(SpawnTrack->CreateNewSection());
+
+		SpawnTrack->AddSection(*MovieSceneSection);
+		SpawnTrack->SetObjectId(Guid);
+
+		MovieSceneSection->SetDefault(false);
+		MovieSceneSection->AddKey(0.0f, false, EMovieSceneKeyInterpolation::Break);
+
+		MovieSceneSection->SetStartTime(Time);
+	}
+
+	bRecording = bRecord;
+
+	bWasSpawned = false;
+
+	Record(Time);
+}
+
+void FMovieSceneSpawnPropertyRecorder::FinalizeSection()
+{
+	bRecording = false;
+
+	const bool bSpawned = ObjectToRecord.IsValid();
+	if(bSpawned != bWasSpawned)
+	{
+		MovieSceneSection->AddKey(MovieSceneSection->GetEndTime(), bSpawned, EMovieSceneKeyInterpolation::Break);
+	}
+}
+
+void FMovieSceneSpawnPropertyRecorder::Record(float CurrentTime)
+{
+	if(bRecording)
+	{
+		if(ObjectToRecord.IsValid())
+		{
+			MovieSceneSection->SetEndTime(CurrentTime);
+		}
+
+		const bool bSpawned = ObjectToRecord.IsValid();
+		if(bSpawned != bWasSpawned)
+		{
+			MovieSceneSection->AddKey(CurrentTime, bSpawned, EMovieSceneKeyInterpolation::Break);
+		}
+		bWasSpawned = bSpawned;
+	}
+}
\ No newline at end of file
diff --git a/Engine/Plugins/Editor/SequenceRecorder/Source/SequenceRecorder/Private/PropertyRecorders/MovieSceneSpawnPropertyRecorder.h b/Engine/Plugins/Editor/SequenceRecorder/Source/SequenceRecorder/Private/PropertyRecorders/MovieSceneSpawnPropertyRecorder.h
new file mode 100644
--- /dev/null
+++ b/Engine/Plugins/Editor/SequenceRecorder/Source/SequenceRecorder/Private/PropertyRecorders/MovieSceneSpawnPropertyRecorder.h
@@ -0,0 +1,27 @@
+// Copyright 1998-2015 Epic Games, Inc. All Rights Reserved.
+
+#pragma once
+
+#include "IMovieScenePropertyRecorder.h"
+
+class FMovieSceneSpawnPropertyRecorder : public IMovieScenePropertyRecorder
+{
+public:
+	virtual ~FMovieSceneSpawnPropertyRecorder() {}
+
+	virtual void CreateSection(UObject* InObjectToRecord, class UMovieScene* MovieScene, const FGuid& Guid, float Time, bool bRecord) override;
+	virtual void FinalizeSection() override;
+	virtual void Record(float CurrentTime) override;
+
+private:
+	/** Object to record from */
+	TLazyObjectPtr<UObject> ObjectToRecord;
+
+	/** Section to record to */
+	TWeakObjectPtr<class UMovieSceneBoolSection> MovieSceneSection;
+
+	/** Flag if we are actually recording or not */
+	bool bRecording;
+
+	bool bWasSpawned;
+};
\ No newline at end of file
diff --git a/Engine/Plugins/Editor/SequenceRecorder/Source/SequenceRecorder/Private/PropertyRecorders/MovieSceneVisibilityPropertyRecorder.cpp b/Engine/Plugins/Editor/SequenceRecorder/Source/SequenceRecorder/Private/PropertyRecorders/MovieSceneVisibilityPropertyRecorder.cpp
--- a/Engine/Plugins/Editor/SequenceRecorder/Source/SequenceRecorder/Private/PropertyRecorders/MovieSceneVisibilityPropertyRecorder.cpp
+++ b/Engine/Plugins/Editor/SequenceRecorder/Source/SequenceRecorder/Private/PropertyRecorders/MovieSceneVisibilityPropertyRecorder.cpp
@@ -8,9 +8,9 @@
 
 static const FName VisibilityTrackName = TEXT("Visibility");
 
-void FMovieSceneVisibilityPropertyRecorder::CreateSection(AActor* SourceActor, UMovieScene* MovieScene, const FGuid& Guid, bool bRecord)
+void FMovieSceneVisibilityPropertyRecorder::CreateSection(UObject* InObjectToRecord, UMovieScene* MovieScene, const FGuid& Guid, float Time, bool bRecord)
 {
-	ActorToRecord = SourceActor;
+	ObjectToRecord = InObjectToRecord;
 
 	UMovieSceneVisibilityTrack* VisibilityTrack = MovieScene->AddTrack<UMovieSceneVisibilityTrack>(Guid);
 	if(VisibilityTrack)
@@ -23,12 +23,24 @@ void FMovieSceneVisibilityPropertyRecorder::CreateSection(AActor* SourceActor, U
 
 		MovieSceneSection->SetDefault(false);
 
-		MovieSceneSection->SetStartTime(0.0f);
+		bWasVisible = false;
+		if(USceneComponent* SceneComponent = Cast<USceneComponent>(ObjectToRecord.Get()))
+		{
+			bWasVisible = SceneComponent->IsVisible() && SceneComponent->IsRegistered();
+		}
+		else if(AActor* Actor = Cast<AActor>(ObjectToRecord.Get()))
+		{
+			bWasVisible = !Actor->bHidden;
+		}
+
+		MovieSceneSection->AddKey(Time, bWasVisible, EMovieSceneKeyInterpolation::Break);
+
+		MovieSceneSection->SetStartTime(Time);
 	}
 
 	bRecording = bRecord;
 
-	bWasVisible = false;
+	Record(Time);
 }
 
 void FMovieSceneVisibilityPropertyRecorder::FinalizeSection()
@@ -38,15 +50,22 @@ void FMovieSceneVisibilityPropertyRecorder::FinalizeSection()
 
 void FMovieSceneVisibilityPropertyRecorder::Record(float CurrentTime)
 {
-	if(ActorToRecord.IsValid())
+	if(ObjectToRecord.IsValid())
 	{
 		MovieSceneSection->SetEndTime(CurrentTime);
 
 		if(bRecording)
 		{
-			USkeletalMeshComponent* SkeletalMeshComponent = ActorToRecord->FindComponentByClass<USkeletalMeshComponent>();
+			bool bVisible = false;
+			if(USceneComponent* SceneComponent = Cast<USceneComponent>(ObjectToRecord.Get()))
+			{
+				bVisible = SceneComponent->IsVisible() && SceneComponent->IsRegistered();
+			}
+			else if(AActor* Actor = Cast<AActor>(ObjectToRecord.Get()))
+			{
+				bVisible = !Actor->bHidden;
+			}
 
-			const bool bVisible = SkeletalMeshComponent ? !SkeletalMeshComponent->bHiddenInGame : !ActorToRecord->bHidden;
 			if(bVisible != bWasVisible)
 			{
 				MovieSceneSection->AddKey(CurrentTime, bVisible, EMovieSceneKeyInterpolation::Break);
diff --git a/Engine/Plugins/Editor/SequenceRecorder/Source/SequenceRecorder/Private/PropertyRecorders/MovieSceneVisibilityPropertyRecorder.h b/Engine/Plugins/Editor/SequenceRecorder/Source/SequenceRecorder/Private/PropertyRecorders/MovieSceneVisibilityPropertyRecorder.h
--- a/Engine/Plugins/Editor/SequenceRecorder/Source/SequenceRecorder/Private/PropertyRecorders/MovieSceneVisibilityPropertyRecorder.h
+++ b/Engine/Plugins/Editor/SequenceRecorder/Source/SequenceRecorder/Private/PropertyRecorders/MovieSceneVisibilityPropertyRecorder.h
@@ -9,16 +9,20 @@ class FMovieSceneVisibilityPropertyRecorder : public IMovieScenePropertyRecorder
 public:
 	virtual ~FMovieSceneVisibilityPropertyRecorder() {}
 
-	virtual void CreateSection(class AActor* SourceActor, class UMovieScene* MovieScene, const FGuid& Guid, bool bRecord) override;
+	virtual void CreateSection(UObject* InObjectToRecord, class UMovieScene* MovieScene, const FGuid& Guid, float Time, bool bRecord) override;
 	virtual void FinalizeSection() override;
 	virtual void Record(float CurrentTime) override;
 
 private:
-	TLazyObjectPtr<class AActor> ActorToRecord;
+	/** Object to record from */
+	TLazyObjectPtr<class UObject> ObjectToRecord;
 
+	/** Section to record to */
 	TWeakObjectPtr<class UMovieSceneVisibilitySection> MovieSceneSection;
 
+	/** Flag if we are actually recording or not */
 	bool bRecording;
 
+	/** Flag used to track visibility state and add keys when this changes */
 	bool bWasVisible;
 };
\ No newline at end of file
diff --git a/Engine/Plugins/Editor/SequenceRecorder/Source/SequenceRecorder/Private/SequenceRecorder.cpp b/Engine/Plugins/Editor/SequenceRecorder/Source/SequenceRecorder/Private/SequenceRecorder.cpp
--- a/Engine/Plugins/Editor/SequenceRecorder/Source/SequenceRecorder/Private/SequenceRecorder.cpp
+++ b/Engine/Plugins/Editor/SequenceRecorder/Source/SequenceRecorder/Private/SequenceRecorder.cpp
@@ -6,19 +6,10 @@
 #include "SequenceRecorderUtils.h"
 #include "SequenceRecorderSettings.h"
 #include "ObjectTools.h"
-#include "LevelSequence.h"
-#include "AssetSelection.h"
-#include "MovieSceneSkeletalAnimationTrack.h"
-#include "MovieSceneSkeletalAnimationSection.h"
-#include "MovieSceneSpawnTrack.h"
-#include "MovieScene3DTransformTrack.h"
-#include "MovieScene3DTransformSection.h"
-#include "MovieSceneBoolSection.h"
 #include "ScopedTransaction.h"
 #include "AssetRegistryModule.h"
-#include "IMovieScenePropertyRecorder.h"
-#include "MovieScene3DTransformPropertyRecorder.h"
-#include "MovieSceneVisibilityPropertyRecorder.h"
+#include "SequenceRecorderActorFilter.h"
+#include "Engine/DemoNetDriver.h"
 
 FSequenceRecorder::FSequenceRecorder()
 	: bQueuedRecordingsDirty(false)
@@ -47,11 +38,24 @@ bool FSequenceRecorder::IsRecordingQueued(AActor* Actor) const
 	return false;
 }
 
+UActorRecording* FSequenceRecorder::FindRecording(AActor* Actor) const
+{
+	for (UActorRecording* QueuedRecording : QueuedRecordings)
+	{
+		if (QueuedRecording->ActorToRecord == Actor)
+		{
+			return QueuedRecording;
+		}
+	}
+
+	return nullptr;
+}
+
 void FSequenceRecorder::StartAllQueuedRecordings()
 {
 	for (UActorRecording* QueuedRecording : QueuedRecordings)
 	{
-		QueuedRecording->StartRecording();
+		QueuedRecording->StartRecording(CurrentSequence.Get(), CurrentTime);
 	}
 }
 
@@ -59,11 +63,11 @@ void FSequenceRecorder::StopAllQueuedRecordings()
 {
 	for (UActorRecording* QueuedRecording : QueuedRecordings)
 	{
-		QueuedRecording->StopRecording();
+		QueuedRecording->StopRecording(CurrentSequence.Get());
 	}
 }
 
-void FSequenceRecorder::AddNewQueuedRecording(AActor* Actor, UAnimSequence* AnimSequence, float Length)
+UActorRecording* FSequenceRecorder::AddNewQueuedRecording(AActor* Actor, UAnimSequence* AnimSequence, float Length)
 {
 	const USequenceRecorderSettings* Settings = GetDefault<USequenceRecorderSettings>();
 
@@ -72,12 +76,16 @@ void FSequenceRecorder::AddNewQueuedRecording(AActor* Actor, UAnimSequence* Anim
 	ActorRecording->ActorToRecord = Actor;
 	ActorRecording->TargetAnimation = AnimSequence;
 	ActorRecording->AnimationSettings.Length = Length;
-	ActorRecording->AnimationSettings.bRecordInWorldSpace = !Settings->bRecordTransformsAsSeperateTrack;
-	ActorRecording->ActorSettings.bRecordTransforms = Settings->bRecordTransformsAsSeperateTrack;
+
+	// We always record in world space as we need animations to record root motion
+	ActorRecording->AnimationSettings.bRecordInWorldSpace = true;
+	ActorRecording->ActorSettings.bRecordTransforms = true;
 
 	QueuedRecordings.Add(ActorRecording);
 
 	bQueuedRecordingsDirty = true;
+
+	return ActorRecording;
 }
 
 void FSequenceRecorder::RemoveQueuedRecording(AActor* Actor)
@@ -118,11 +126,43 @@ bool FSequenceRecorder::HasQueuedRecordings() const
 
 void FSequenceRecorder::Tick(float DeltaSeconds)
 {
+	// if a replay recording is in progress and channels are paused, wait until we have data again before recording
+	if(CurrentReplayWorld.IsValid() && CurrentReplayWorld->DemoNetDriver != nullptr)
+	{
+		if(CurrentReplayWorld->DemoNetDriver->bChannelsArePaused)
+		{
+			return;
+		}
+	}
+
+	const USequenceRecorderSettings* Settings = GetDefault<USequenceRecorderSettings>();
+
+	// check for spawned actors and if they have begun playing yet
+	for(TWeakObjectPtr<AActor>& QueuedSpawnedActor : QueuedSpawnedActors)
+	{
+		if(AActor* Actor = QueuedSpawnedActor.Get())
+		{
+			if(Actor->HasActorBegunPlay())
+			{
+				if(UActorRecording::IsRelevantForRecording(Actor) && IsActorValidForRecording(Actor))
+				{
+					UActorRecording* NewRecording = AddNewQueuedRecording(Actor);
+					NewRecording->bWasSpawnedPostRecord = true;
+					NewRecording->StartRecording(CurrentSequence.Get(), CurrentTime);
+				}
+
+				QueuedSpawnedActor.Reset();
+			}
+		}
+	}
+
+	QueuedSpawnedActors.RemoveAll([](TWeakObjectPtr<AActor>& QueuedSpawnedActor) { return !QueuedSpawnedActor.IsValid(); } );
+
 	FAnimationRecorderManager::Get().Tick(DeltaSeconds);
 
 	for(UActorRecording* Recording : QueuedRecordings)
 	{
-		Recording->Tick(DeltaSeconds);
+		Recording->Tick(DeltaSeconds, CurrentSequence.Get(), CurrentTime);
 	}
 
 	if(CurrentDelay > 0.0f)
@@ -131,22 +171,15 @@ void FSequenceRecorder::Tick(float DeltaSeconds)
 		if(CurrentDelay <= 0.0f)
 		{
 			CurrentDelay = 0.0f;
-			StartRecordingInternal();
+			StartRecordingInternal(nullptr);
 		}
 	}
 
-	const USequenceRecorderSettings* Settings = GetDefault<USequenceRecorderSettings>();
 	if(Settings->bCreateLevelSequence && CurrentSequence.IsValid())
 	{
-		// record our property recorders
-		for(auto& PropertyRecorder : PropertyRecorders)
-		{
-			PropertyRecorder->Record(CurrentTime);
-		}
-
 		CurrentTime += DeltaSeconds;
 
-		// check if all our actor recordings are finished
+		// check if all our actor recordings are finished or we timed out
 		if(QueuedRecordings.Num() > 0)
 		{
 			bool bAllFinished = true;
@@ -159,11 +192,29 @@ void FSequenceRecorder::Tick(float DeltaSeconds)
 				}
 			}
 
-			if(bAllFinished)
+			if(bAllFinished || (Settings->SequenceLength > 0.0f && CurrentTime >= Settings->SequenceLength))
 			{
 				StopRecording();
 			}
 		}
+
+		auto RemoveDeadActorPredicate = 
+			[&](UActorRecording* Recording)
+			{
+				if(!Recording->ActorToRecord.IsValid())
+				{
+					DeadRecordings.Add(Recording);
+					return true;
+				}
+				return false; 
+			};
+
+		// remove all dead actors
+		int32 Removed = QueuedRecordings.RemoveAll(RemoveDeadActorPredicate);
+		if(Removed > 0)
+		{
+			bQueuedRecordingsDirty = true;
+		}
 	}
 }
 
@@ -171,20 +222,55 @@ bool FSequenceRecorder::StartRecording()
 {
 	const USequenceRecorderSettings* Settings = GetDefault<USequenceRecorderSettings>();
 
+	CurrentTime = 0.0f;
+
 	if(Settings->RecordingDelay > 0.0f)
 	{
 		CurrentDelay = Settings->RecordingDelay;
+		
 		return QueuedRecordings.Num() > 0;
 	}
 	
-	return StartRecordingInternal();
+	return StartRecordingInternal(nullptr);
 }
 
-bool FSequenceRecorder::StartRecordingInternal()
+bool FSequenceRecorder::StartRecordingForReplay(UWorld* World, const FSequenceRecorderActorFilter& ActorFilter)
+{
+	// set up our recording settings
+	USequenceRecorderSettings* Settings = GetMutableDefault<USequenceRecorderSettings>();
+	
+	Settings->bCreateLevelSequence = true;
+	Settings->SequenceLength = 0.0f;
+	Settings->RecordingDelay = 0.0f;
+	Settings->bRecordNearbySpawnedActors = true;
+	Settings->NearbyActorRecordingProximity = 0.0f;
+	Settings->bRecordWorldSettingsActor = true;
+	Settings->ActorFilter = ActorFilter;
+
+	CurrentReplayWorld = World;
+	
+	return StartRecordingInternal(World);
+}
+
+bool FSequenceRecorder::StartRecordingInternal(UWorld* World)
 {
 	CurrentTime = 0.0f;
 
 	const USequenceRecorderSettings* Settings = GetDefault<USequenceRecorderSettings>();
+
+	if(Settings->bRecordWorldSettingsActor)
+	{
+		if(World != nullptr || (QueuedRecordings.Num() > 0 && QueuedRecordings[0]->ActorToRecord.IsValid()))
+		{
+			UWorld* ActorWorld = World != nullptr ? World : QueuedRecordings[0]->ActorToRecord->GetWorld();
+			if(ActorWorld)
+			{
+				AWorldSettings* WorldSettings = ActorWorld->GetWorldSettings();
+				AddNewQueuedRecording(WorldSettings);
+			}
+		}
+	}
+
 	if(QueuedRecordings.Num() > 0)
 	{
 		ULevelSequence* LevelSequence = nullptr;
@@ -202,6 +288,26 @@ bool FSequenceRecorder::StartRecordingInternal()
 			{
 				LevelSequence->Initialize();
 				CurrentSequence = LevelSequence;
+
+				FAssetRegistryModule::AssetCreated(LevelSequence);
+			}
+		}
+
+		// register for spawning delegate in the world(s) of recorded actors
+		for(UActorRecording* Recording : QueuedRecordings)
+		{
+			if(Recording->ActorToRecord.IsValid())
+			{
+				UWorld* ActorWorld = Recording->ActorToRecord->GetWorld();
+				if(ActorWorld != nullptr)
+				{
+					FDelegateHandle* FoundHandle = ActorSpawningDelegateHandles.Find(ActorWorld);
+					if(FoundHandle == nullptr)
+					{
+						FDelegateHandle NewHandle = ActorWorld->AddOnActorSpawnedHandler(FOnActorSpawned::FDelegate::CreateStatic(&FSequenceRecorder::HandleActorSpawned));
+						ActorSpawningDelegateHandles.Add(ActorWorld, NewHandle);
+					}
+				}
 			}
 		}
 
@@ -209,56 +315,11 @@ bool FSequenceRecorder::StartRecordingInternal()
 		bool bStartedRecordingAllActors = true;
 		for(UActorRecording* Recording : QueuedRecordings)
 		{
-			if(!Recording->StartRecording())
+			if(!Recording->StartRecording(CurrentSequence.Get(), CurrentTime))
 			{
 				bStartedRecordingAllActors = false;
 				break;
 			}
-
-			// set up our spawnable for this actor
-			UMovieScene* MovieScene = LevelSequence->GetMovieScene();
-			UAnimSequence* AnimSequence = Recording->LastRecordedAnimation.Get();
-			UActorFactory* FactoryToUse = FActorFactoryAssetProxy::GetFactoryForAssetObject(AnimSequence);
-			check(FactoryToUse);
-
-			FString BlueprintName = AnimSequence->GetName();
-			auto DuplName = [&](FMovieSceneSpawnable& InSpawnable)
-			{
-				return InSpawnable.GetName() == BlueprintName;
-			};
-
-			int32 Index = 2;
-			FString UniqueString;
-			while (MovieScene->FindSpawnable(DuplName))
-			{
-				BlueprintName.RemoveFromEnd(UniqueString);
-				UniqueString = FString::Printf(TEXT(" (%d)"), Index);
-				BlueprintName += UniqueString;
-			}
-
-			UBlueprint* Blueprint = FactoryToUse->CreateBlueprint(AnimSequence, MovieScene, *BlueprintName);
-			check(Blueprint);
-
-			Recording->Guid = MovieScene->AddSpawnable(BlueprintName, Blueprint);
-			
-			// add the anim track
-			UMovieSceneSkeletalAnimationTrack* AnimTrack = MovieScene->AddTrack<UMovieSceneSkeletalAnimationTrack>(Recording->Guid);
-			if(AnimTrack)
-			{
-				AnimTrack->AddNewAnimation(0.0f, AnimSequence);
-			}
-
-			// now add tracks to record
-
-			// we need to create a transform track even if we arent recording transforms
-			TSharedPtr<FMovieScene3DTransformPropertyRecorder> TransformRecorder = MakeShareable(new FMovieScene3DTransformPropertyRecorder);
-			TransformRecorder->CreateSection(Recording->ActorToRecord.Get(), MovieScene, Recording->Guid, Recording->ActorSettings.bRecordTransforms);
-			PropertyRecorders.Add(TransformRecorder);
-
-			// add a visibility track 
-			TSharedPtr<FMovieSceneVisibilityPropertyRecorder> VisibilityRecorder = MakeShareable(new FMovieSceneVisibilityPropertyRecorder);
-			VisibilityRecorder->CreateSection(Recording->ActorToRecord.Get(), MovieScene, Recording->Guid ,Recording->ActorSettings.bRecordVisibility);
-			PropertyRecorders.Add(VisibilityRecorder);
 		}
 
 		if(!bStartedRecordingAllActors)
@@ -272,11 +333,7 @@ bool FSequenceRecorder::StartRecordingInternal()
 
 			for(UActorRecording* Recording : QueuedRecordings)
 			{
-				Recording->StopRecording();
-				if(Recording->LastRecordedAnimation.IsValid())
-				{
-					AssetsToCleanUp.Add(FAssetData(Recording->LastRecordedAnimation.Get()));
-				}
+				Recording->StopRecording(CurrentSequence.Get());
 			}
 
 			// clean up any assets that we can
@@ -296,72 +353,62 @@ bool FSequenceRecorder::StopRecording()
 {
 	CurrentDelay = 0.0f;
 
-	for(UActorRecording* Recording : QueuedRecordings)
+	// remove spawn delegates
+	for(auto It = ActorSpawningDelegateHandles.CreateConstIterator(); It; ++It)
 	{
-		Recording->StopRecording();
+		TWeakObjectPtr<UWorld> World = It->Key;
+		if(World.IsValid())
+		{
+			World->RemoveOnActorSpawnedHandler(It->Value);
+		}
 	}
 
+	ActorSpawningDelegateHandles.Empty();
+
 	// also stop all dead animation recordings, i.e. ones that use GCed components
 	FAnimationRecorderManager::Get().StopRecordingDeadAnimations();
 
+	for(UActorRecording* Recording : QueuedRecordings)
+	{
+		Recording->StopRecording(CurrentSequence.Get());
+	}
+
+	for(UActorRecording* Recording : DeadRecordings)
+	{
+		Recording->StopRecording(CurrentSequence.Get());
+	}
+
+	DeadRecordings.Empty();
+
 	const USequenceRecorderSettings* Settings = GetDefault<USequenceRecorderSettings>();
 	if(Settings->bCreateLevelSequence)
 	{
 		if(CurrentSequence.IsValid())
 		{
-			for(auto& PropertyRecorder : PropertyRecorders)
-			{
-				PropertyRecorder->FinalizeSection();
-			}
-
 			ULevelSequence* LevelSequence = CurrentSequence.Get();
 
 			// Stop referencing the sequence so we are listed as 'not recording'
 			CurrentSequence = nullptr;
 
+			// set movie scene playback range to encompass all sections
 			UMovieScene* MovieScene = LevelSequence->GetMovieScene();
 			if(MovieScene)
 			{
-				float MinRange = 0.0f;
 				float MaxRange = 0.0f;
+				float MinRange = 0.0f;
 
-				for(UActorRecording* Recording : QueuedRecordings)
+				TArray<UMovieSceneSection*> MovieSceneSections = MovieScene->GetAllSections();
+				for(UMovieSceneSection* Section : MovieSceneSections)
 				{
-					if(Recording->LastRecordedAnimation.IsValid())
-					{
-						UAnimSequence* AnimSequence = Recording->LastRecordedAnimation.Get();
-
-						// resize anim sequence section to fit now it is recorded
-						UMovieSceneSkeletalAnimationTrack* AnimTrack = MovieScene->FindTrack<UMovieSceneSkeletalAnimationTrack>(Recording->Guid);
-						check(AnimTrack);
-					
-						UMovieSceneSkeletalAnimationSection* AnimSection = Cast<UMovieSceneSkeletalAnimationSection>(AnimTrack->GetAllSections()[0]);
-						AnimSection->SetEndTime(AnimSequence->GetPlayLength());
-
-						MaxRange = FMath::Max(MaxRange, AnimSequence->GetPlayLength());
-
-						// add a spawn track so we actually get spawned
-						UMovieSceneSpawnTrack* SpawnTrack = MovieScene->AddTrack<UMovieSceneSpawnTrack>(Recording->Guid);
-						if(SpawnTrack)
-						{
-							UMovieSceneBoolSection* SpawnSection = Cast<UMovieSceneBoolSection>(SpawnTrack->CreateNewSection());
-
-							SpawnTrack->AddSection(*SpawnSection);
-							SpawnTrack->SetObjectId(Recording->Guid);
-
-							SpawnSection->SetStartTime(0.0f);
-							SpawnSection->SetEndTime(AnimSequence->GetPlayLength());
-
-							SpawnSection->AddKey(0.0f, true, EMovieSceneKeyInterpolation::Break);
-							SpawnSection->AddKey(AnimSequence->GetPlayLength(), false, EMovieSceneKeyInterpolation::Break);
-						}
-					}
+					MaxRange = FMath::Max(MaxRange, Section->GetEndTime());
+					MinRange = FMath::Min(MinRange, Section->GetStartTime());
 				}
 
 				MovieScene->SetPlaybackRange(MinRange, MaxRange);
-
-				FAssetRegistryModule::AssetCreated(MovieScene);
 			}
+
+			// @todo: post-process sequence to remove redundant tracks (empty transforms etc.)
+
 			return true;
 		}
 	}
@@ -381,4 +428,57 @@ bool FSequenceRecorder::IsDelaying() const
 float FSequenceRecorder::GetCurrentDelay() const
 {
 	return CurrentDelay;
+}
+
+bool FSequenceRecorder::IsActorValidForRecording(AActor* Actor)
+{
+	const USequenceRecorderSettings* Settings = GetDefault<USequenceRecorderSettings>();
+
+	float Distance = Settings->NearbyActorRecordingProximity;
+
+	// check distance if valid
+	if(Distance > 0.0f)
+	{
+		const FTransform ActorTransform = Actor->GetTransform();
+		const FVector ActorTranslation = ActorTransform.GetTranslation();
+
+		for(UActorRecording* Recording : QueuedRecordings)
+		{
+			if(AActor* OtherActor = Recording->ActorToRecord.Get())
+			{
+				if(OtherActor != Actor)
+				{
+					const FTransform OtherActorTransform = OtherActor->GetTransform();
+					const FVector OtherActorTranslation = OtherActorTransform.GetTranslation();
+
+					if((OtherActorTranslation - ActorTranslation).Size() < Distance)
+					{
+						return true;
+					}
+				}
+			}
+		}
+	}
+
+	// check class if any
+	for(const TSubclassOf<AActor>& ActorClass : Settings->ActorFilter.ActorClassesToRecord)
+	{
+		if(Actor->IsA(*ActorClass))
+		{
+			return true;
+		}
+	}
+
+	return false;
+}
+
+void FSequenceRecorder::HandleActorSpawned(AActor* Actor)
+{
+	const USequenceRecorderSettings* Settings = GetDefault<USequenceRecorderSettings>();
+
+	if(Actor && FSequenceRecorder::Get().IsRecording() && Settings->bRecordNearbySpawnedActors)
+	{
+		// Queue a possible recording - we need to wait until the actor has begun playing
+		FSequenceRecorder::Get().QueuedSpawnedActors.Add(Actor);
+	}
 }
\ No newline at end of file
diff --git a/Engine/Plugins/Editor/SequenceRecorder/Source/SequenceRecorder/Private/SequenceRecorder.h b/Engine/Plugins/Editor/SequenceRecorder/Source/SequenceRecorder/Private/SequenceRecorder.h
--- a/Engine/Plugins/Editor/SequenceRecorder/Source/SequenceRecorder/Private/SequenceRecorder.h
+++ b/Engine/Plugins/Editor/SequenceRecorder/Source/SequenceRecorder/Private/SequenceRecorder.h
@@ -11,6 +11,9 @@ public:
 	/** Starts recording a sequence. */
 	bool StartRecording();
 
+	/** Starts recording a sequence for the specified world. */
+	bool StartRecordingForReplay(UWorld* World, const struct FSequenceRecorderActorFilter& ActorFilter);
+
 	/** Stops any currently recording sequence */
 	bool StopRecording();
 
@@ -23,15 +26,19 @@ public:
 	/** Get the current delay we are waiting for */
 	float GetCurrentDelay() const;
 
+	TWeakObjectPtr<class ULevelSequence> GetCurrentSequence() { return CurrentSequence; }
+
 	bool IsRecordingQueued(AActor* Actor) const;
 
+	class UActorRecording* FindRecording(AActor* Actor) const;
+
 	void StartAllQueuedRecordings();
 
 	void StopAllQueuedRecordings();
 
 	void StopRecordingDeadAnimations();
 
-	void AddNewQueuedRecording(AActor* Actor = nullptr, UAnimSequence* AnimSequence = nullptr, float Length = 0.0f);
+	class UActorRecording* AddNewQueuedRecording(AActor* Actor = nullptr, UAnimSequence* AnimSequence = nullptr, float Length = 0.0f);
 
 	void RemoveQueuedRecording(AActor* Actor);
 
@@ -45,9 +52,17 @@ public:
 
 	void ResetQueuedRecordingsDirty() { bQueuedRecordingsDirty = false; }
 
+	bool IsRecording() const { return CurrentSequence.IsValid(); }
+
 private:
 	/** Starts recording a sequence, possibly delayed */
-	bool StartRecordingInternal();
+	bool StartRecordingInternal(UWorld* World);
+
+	/** Check if an actor is valid for recording */
+	bool IsActorValidForRecording(AActor* Actor);
+
+	/** Handle actors being spawned */
+	static void HandleActorSpawned(AActor* Actor);
 
 private:
 	/** Constructor, private - use Get() function */
@@ -56,11 +71,15 @@ private:
 	/** Currently recording level sequence, if any */
 	TWeakObjectPtr<class ULevelSequence> CurrentSequence;
 
-	/** The current set of property recorders */
-	TArray<TSharedPtr<class IMovieScenePropertyRecorder>> PropertyRecorders;
+	/** World we are recording a replay for, if any */
+	TLazyObjectPtr<class UWorld> CurrentReplayWorld;
 
 	TArray<class UActorRecording*> QueuedRecordings;
 
+	TArray<class UActorRecording*> DeadRecordings;
+
+	TArray<TWeakObjectPtr<AActor>> QueuedSpawnedActors;
+
 	bool bQueuedRecordingsDirty;
 
 	/** The delay we are currently waiting for */
@@ -68,4 +87,7 @@ private:
 
 	/** Current recording time */
 	float CurrentTime;
+
+	/** Delegate handles for FOnActorSpawned events */
+	TMap<TWeakObjectPtr<UWorld>, FDelegateHandle> ActorSpawningDelegateHandles;
 };
diff --git a/Engine/Plugins/Editor/SequenceRecorder/Source/SequenceRecorder/Private/SequenceRecorderModule.cpp b/Engine/Plugins/Editor/SequenceRecorder/Source/SequenceRecorder/Private/SequenceRecorderModule.cpp
--- a/Engine/Plugins/Editor/SequenceRecorder/Source/SequenceRecorder/Private/SequenceRecorderModule.cpp
+++ b/Engine/Plugins/Editor/SequenceRecorder/Source/SequenceRecorder/Private/SequenceRecorderModule.cpp
@@ -8,12 +8,13 @@
 #include "SequenceRecorderCommands.h"
 #include "SequenceRecorder.h"
 #include "SequenceRecorderSettings.h"
+#include "ISequenceRecorder.h"
 
 #define LOCTEXT_NAMESPACE "SequenceRecorder"
 
 static const FName SequenceRecorderTabName(TEXT("SequenceRecorder"));
 
-class FSequenceRecorderModule : public IModuleInterface, private FSelfRegisteringExec
+class FSequenceRecorderModule : public ISequenceRecorder, private FSelfRegisteringExec
 {
 	/** IModuleInterface implementation */
 	virtual void StartupModule() override
@@ -213,6 +214,22 @@ class FSequenceRecorderModule : public IModuleInterface, private FSelfRegisterin
 		return false;
 	}
 
+	// ISequenceRecorder interface
+	virtual bool StartRecording(UWorld* World, const FSequenceRecorderActorFilter& ActorFilter) override
+	{
+		return FSequenceRecorder::Get().StartRecordingForReplay(World, ActorFilter);
+	}
+
+	virtual void StopRecording() override
+	{
+		FSequenceRecorder::Get().StopRecording();
+	}
+
+	virtual bool IsRecording() override
+	{
+		return FSequenceRecorder::Get().IsRecording();
+	}
+
 	static void TickSequenceRecorder(float DeltaSeconds)
 	{
 		if (!IsRunningDedicatedServer() && !IsRunningCommandlet())
diff --git a/Engine/Plugins/Editor/SequenceRecorder/Source/SequenceRecorder/Private/SequenceRecorderPrivatePCH.h b/Engine/Plugins/Editor/SequenceRecorder/Source/SequenceRecorder/Private/SequenceRecorderPrivatePCH.h
--- a/Engine/Plugins/Editor/SequenceRecorder/Source/SequenceRecorder/Private/SequenceRecorderPrivatePCH.h
+++ b/Engine/Plugins/Editor/SequenceRecorder/Source/SequenceRecorder/Private/SequenceRecorderPrivatePCH.h
@@ -11,4 +11,5 @@
 #include "MovieScene.h"
 #include "KeyParams.h"
 #include "IKeyframeSection.h"
-#include "MovieSceneSection.h"
\ No newline at end of file
+#include "MovieSceneSection.h"
+#include "LevelSequence.h"
\ No newline at end of file
diff --git a/Engine/Plugins/Editor/SequenceRecorder/Source/SequenceRecorder/Private/SequenceRecorderSettings.cpp b/Engine/Plugins/Editor/SequenceRecorder/Source/SequenceRecorder/Private/SequenceRecorderSettings.cpp
--- a/Engine/Plugins/Editor/SequenceRecorder/Source/SequenceRecorder/Private/SequenceRecorderSettings.cpp
+++ b/Engine/Plugins/Editor/SequenceRecorder/Source/SequenceRecorder/Private/SequenceRecorderSettings.cpp
@@ -12,4 +12,7 @@ USequenceRecorderSettings::USequenceRecorderSettings(const FObjectInitializer& O
 	SequenceName = TEXT("RecordedSequence");
 	AnimationSubDirectory = TEXT("Animations");
 	SequenceRecordingBasePath.Path = TEXT("Cinematics/Sequences");
+	bRecordNearbySpawnedActors = false;
+	NearbyActorRecordingProximity = 1000.0f;
+	bRecordWorldSettingsActor = false;
 }
